[AudioModulation] Moved Audio Modulation Insights code from Audio Modulation plugin to its own plugin so Audio Modulation does not depend on Audio Insights plugin + code adjustments

diff --git a/Engine/Plugins/Runtime/AudioModulation/AudioModulation.uplugin b/Engine/Plugins/Runtime/AudioModulation/AudioModulation.uplugin
--- a/Engine/Plugins/Runtime/AudioModulation/AudioModulation.uplugin
+++ b/Engine/Plugins/Runtime/AudioModulation/AudioModulation.uplugin
@@ -35,10 +35,6 @@
 		{
 			"Name": "WaveTable",
 			"Enabled": true
-		},
-		{
-			"Name": "AudioInsights",
-			"Enabled": true
 		}
 	]
 }
\ No newline at end of file
diff --git a/Engine/Plugins/Runtime/AudioModulation/Source/AudioModulationEditor/AudioModulationEditor.build.cs b/Engine/Plugins/Runtime/AudioModulation/Source/AudioModulationEditor/AudioModulationEditor.build.cs
--- a/Engine/Plugins/Runtime/AudioModulation/Source/AudioModulationEditor/AudioModulationEditor.build.cs
+++ b/Engine/Plugins/Runtime/AudioModulation/Source/AudioModulationEditor/AudioModulationEditor.build.cs
@@ -28,13 +28,10 @@ public class AudioModulationEditor : ModuleRules
 				"InputCore",				
 				"AudioEditor",
 				"AudioExtensions",
-				"AudioInsights",
-				"AudioInsightsEditor",
 				"AudioModulation",
 				"CurveEditor",
 				"EditorWidgets",
 				"ToolWidgets",
-				"TraceAnalysis",
 				"WaveTable",
 				"WaveTableEditor",
 				"MovieScene",
diff --git a/Engine/Plugins/Runtime/AudioModulation/Source/AudioModulationEditor/Private/AudioModulationEditor.cpp b/Engine/Plugins/Runtime/AudioModulation/Source/AudioModulationEditor/Private/AudioModulationEditor.cpp
--- a/Engine/Plugins/Runtime/AudioModulation/Source/AudioModulationEditor/Private/AudioModulationEditor.cpp
+++ b/Engine/Plugins/Runtime/AudioModulation/Source/AudioModulationEditor/Private/AudioModulationEditor.cpp
@@ -13,11 +13,8 @@
 #include "Framework/Commands/UIAction.h"
 #include "Framework/MultiBox/MultiBoxBuilder.h"
 #include "Framework/MultiBox/MultiBoxExtender.h"
-#include "IAudioInsightsEditorModule.h"
 #include "IAudioModulation.h"
 #include "ICurveEditorModule.h"
-#include "Insights/Views/ControlBusDashboardViewFactory.h"
-#include "Insights/Views/ModulationMatrixDashboardViewFactory.h"
 #include "Internationalization/Internationalization.h"
 #include "Layouts/SoundControlBusMixStageLayout.h"
 #include "Layouts/SoundControlModulationPatchLayout.h"
@@ -56,42 +53,6 @@ namespace AudioModulationEditor
 	}
 } // namespace AudioModulationEditor
 
-namespace
-{
-	static bool bIsModulationRegisteredInAudioInsights = false;
-
-	void RegisterModulationInAudioInsights()
-	{
-		// Don't run providers in any commandlet to avoid additional, unnecessary overhead as audio modulation insights is dormant.
-		if (!IsRunningCommandlet())
-		{
-			if (!bIsModulationRegisteredInAudioInsights)
-			{
-				IAudioInsightsEditorModule& InsightsEditorModule = IAudioInsightsEditorModule::GetChecked();
-				InsightsEditorModule.RegisterDashboardViewFactory(MakeShared<AudioModulationEditor::FControlBusDashboardViewFactory>());
-				InsightsEditorModule.RegisterDashboardViewFactory(MakeShared<AudioModulationEditor::FModulationMatrixDashboardViewFactory>());
-
-				bIsModulationRegisteredInAudioInsights = true;
-			}
-		}
-	}
-
-	void UnregisterModulationInAudioInsights()
-	{
-		if (bIsModulationRegisteredInAudioInsights)
-		{
-			if (IAudioInsightsEditorModule::IsModuleLoaded())
-			{
-				IAudioInsightsEditorModule& InsightsEditorModule = IAudioInsightsEditorModule::GetChecked();
-				InsightsEditorModule.UnregisterDashboardViewFactory("ControlBuses");
-				InsightsEditorModule.UnregisterDashboardViewFactory("ModulationMatrix");
-
-				bIsModulationRegisteredInAudioInsights = false;
-			}
-		}
-	}
-}
-
 FAudioModulationEditorModule::FAudioModulationEditorModule()
 {
 }
@@ -203,8 +164,6 @@ void FAudioModulationEditorModule::StartupModule()
 			Parameter->RemoveFromRoot();
 		}
 	});
-
-	::RegisterModulationInAudioInsights();
 }
 
 void FAudioModulationEditorModule::RegisterCustomPropertyLayouts()
@@ -255,8 +214,6 @@ void FAudioModulationEditorModule::ShutdownModule()
 		SequencerModule.UnRegisterTrackEditor(AudioControlBusTrackCreateEditorHandle);
 	}
 #endif
-
-	::UnregisterModulationInAudioInsights();
 }
 
 IMPLEMENT_MODULE(FAudioModulationEditorModule, AudioModulationEditor);
diff --git a/Engine/Plugins/Runtime/AudioModulation/Source/AudioModulationEditor/Private/Insights/Messages/ControlBusTraceMessages.h b/Engine/Plugins/Runtime/AudioModulation/Source/AudioModulationEditor/Private/Insights/Messages/ControlBusTraceMessages.h
deleted file mode 100644
diff --git a/Engine/Plugins/Runtime/AudioModulation/Source/AudioModulationEditor/Private/Insights/Messages/ModulationMatrixTraceMessages.h b/Engine/Plugins/Runtime/AudioModulation/Source/AudioModulationEditor/Private/Insights/Messages/ModulationMatrixTraceMessages.h
deleted file mode 100644
diff --git a/Engine/Plugins/Runtime/AudioModulation/Source/AudioModulationEditor/Private/Insights/Providers/ControlBusTraceProvider.cpp b/Engine/Plugins/Runtime/AudioModulation/Source/AudioModulationEditor/Private/Insights/Providers/ControlBusTraceProvider.cpp
deleted file mode 100644
diff --git a/Engine/Plugins/Runtime/AudioModulation/Source/AudioModulationEditor/Private/Insights/Providers/ControlBusTraceProvider.h b/Engine/Plugins/Runtime/AudioModulation/Source/AudioModulationEditor/Private/Insights/Providers/ControlBusTraceProvider.h
deleted file mode 100644
diff --git a/Engine/Plugins/Runtime/AudioModulation/Source/AudioModulationEditor/Private/Insights/Providers/ModulationMatrixTraceProvider.cpp b/Engine/Plugins/Runtime/AudioModulation/Source/AudioModulationEditor/Private/Insights/Providers/ModulationMatrixTraceProvider.cpp
deleted file mode 100644
diff --git a/Engine/Plugins/Runtime/AudioModulation/Source/AudioModulationEditor/Private/Insights/Providers/ModulationMatrixTraceProvider.h b/Engine/Plugins/Runtime/AudioModulation/Source/AudioModulationEditor/Private/Insights/Providers/ModulationMatrixTraceProvider.h
deleted file mode 100644
diff --git a/Engine/Plugins/Runtime/AudioModulation/Source/AudioModulationEditor/Private/Insights/Views/AudioModulationEditorCommands.cpp b/Engine/Plugins/Runtime/AudioModulation/Source/AudioModulationEditor/Private/Insights/Views/AudioModulationEditorCommands.cpp
deleted file mode 100644
diff --git a/Engine/Plugins/Runtime/AudioModulation/Source/AudioModulationEditor/Private/Insights/Views/AudioModulationEditorCommands.h b/Engine/Plugins/Runtime/AudioModulation/Source/AudioModulationEditor/Private/Insights/Views/AudioModulationEditorCommands.h
deleted file mode 100644
diff --git a/Engine/Plugins/Runtime/AudioModulation/Source/AudioModulationEditor/Private/Insights/Views/ControlBusDashboardViewFactory.cpp b/Engine/Plugins/Runtime/AudioModulation/Source/AudioModulationEditor/Private/Insights/Views/ControlBusDashboardViewFactory.cpp
deleted file mode 100644
diff --git a/Engine/Plugins/Runtime/AudioModulation/Source/AudioModulationEditor/Private/Insights/Views/ControlBusDashboardViewFactory.h b/Engine/Plugins/Runtime/AudioModulation/Source/AudioModulationEditor/Private/Insights/Views/ControlBusDashboardViewFactory.h
deleted file mode 100644
diff --git a/Engine/Plugins/Runtime/AudioModulation/Source/AudioModulationEditor/Private/Insights/Views/ModulationMatrixDashboardViewFactory.cpp b/Engine/Plugins/Runtime/AudioModulation/Source/AudioModulationEditor/Private/Insights/Views/ModulationMatrixDashboardViewFactory.cpp
deleted file mode 100644
diff --git a/Engine/Plugins/Runtime/AudioModulation/Source/AudioModulationEditor/Private/Insights/Views/ModulationMatrixDashboardViewFactory.h b/Engine/Plugins/Runtime/AudioModulation/Source/AudioModulationEditor/Private/Insights/Views/ModulationMatrixDashboardViewFactory.h
deleted file mode 100644
diff --git a/Engine/Plugins/Runtime/AudioModulationInsights/AudioModulationInsights.uplugin b/Engine/Plugins/Runtime/AudioModulationInsights/AudioModulationInsights.uplugin
new file mode 100644
--- /dev/null
+++ b/Engine/Plugins/Runtime/AudioModulationInsights/AudioModulationInsights.uplugin
@@ -0,0 +1,38 @@
+{
+	"FileVersion": 3,
+	"Version": 1,
+	"VersionName": "1.0",
+	"FriendlyName": "AudioModulationInsights",
+	"Description": " Audio Modulation dashboards for Audio Insights.",
+	"Category": "Audio",
+	"CreatedBy": "Epic Games, Inc.",
+	"CreatedByURL": "https://epicgames.com",
+	"DocsURL": "",
+	"MarketplaceURL": "",
+	"SupportURL": "",
+	"EnabledByDefault": false,
+	"CanContainContent": true,
+	"IsBetaVersion": false,
+	"IsExperimentalVersion": true,
+	"Installed": false,
+	"SupportedPrograms": [ "UnrealInsights" ],
+	"Modules":
+	[
+		{
+			"Name": "AudioModulationInsights",
+			"Type": "EditorAndProgram",
+			"LoadingPhase": "Default",
+			"ProgramAllowList":
+			[
+				"UnrealInsights"
+			]
+		}
+	],
+	"Plugins":
+	[
+		{
+			"Name": "AudioInsights",
+			"Enabled": true
+		}
+	]
+}
diff --git a/Engine/Plugins/Runtime/AudioModulationInsights/Source/AudioModulationInsights/AudioModulationInsights.Build.cs b/Engine/Plugins/Runtime/AudioModulationInsights/Source/AudioModulationInsights/AudioModulationInsights.Build.cs
new file mode 100644
--- /dev/null
+++ b/Engine/Plugins/Runtime/AudioModulationInsights/Source/AudioModulationInsights/AudioModulationInsights.Build.cs
@@ -0,0 +1,41 @@
+// Copyright Epic Games, Inc. All Rights Reserved.
+
+using UnrealBuildTool;
+
+public class AudioModulationInsights : ModuleRules
+{
+	public AudioModulationInsights(ReadOnlyTargetRules Target) : base(Target)
+	{
+		PublicDependencyModuleNames.AddRange(
+			new string[]
+			{
+				"Core",
+			}
+		);
+		
+		PrivateDependencyModuleNames.AddRange(
+			new string[]
+			{
+				"AudioInsights",
+				"AudioMixerCore",
+				"CoreUObject",
+				"InputCore",
+				"Projects",
+				"Slate",
+				"SlateCore",
+				"TraceAnalysis",
+			}
+		);
+
+		if (Target.Type == TargetType.Editor)
+		{
+			PrivateDependencyModuleNames.AddRange(
+				new string[]
+				{
+					"AudioInsightsEditor",
+					"Engine",
+				}
+			);
+		}
+	}
+}
diff --git a/Engine/Plugins/Runtime/AudioModulationInsights/Source/AudioModulationInsights/Private/AudioModulationInsightsModule.cpp b/Engine/Plugins/Runtime/AudioModulationInsights/Source/AudioModulationInsights/Private/AudioModulationInsightsModule.cpp
new file mode 100644
--- /dev/null
+++ b/Engine/Plugins/Runtime/AudioModulationInsights/Source/AudioModulationInsights/Private/AudioModulationInsightsModule.cpp
@@ -0,0 +1,50 @@
+// Copyright Epic Games, Inc. All Rights Reserved.
+#include "AudioModulationInsightsModule.h"
+
+#include "IAudioInsightsModule.h"
+#include "Views/ControlBusDashboardViewFactory.h"
+#include "Views/ModulationMatrixDashboardViewFactory.h"
+
+#if WITH_EDITOR
+#include "IAudioInsightsEditorModule.h"
+#endif // WITH_EDITOR
+
+void FAudioModulationInsightsModule::StartupModule()
+{
+	// Don't run providers in any commandlet to avoid additional, unnecessary overhead as audio modulation insights is dormant.
+	if (!IsRunningCommandlet())
+	{
+#if WITH_EDITOR
+		IAudioInsightsEditorModule& AudioInsightsModule = IAudioInsightsEditorModule::GetChecked();
+#else
+		IAudioInsightsModule& AudioInsightsModule = IAudioInsightsModule::GetChecked();
+#endif // WITH_EDITOR
+		
+		AudioInsightsModule.RegisterDashboardViewFactory(MakeShared<AudioModulationInsights::FControlBusDashboardViewFactory>());
+		AudioInsightsModule.RegisterDashboardViewFactory(MakeShared<AudioModulationInsights::FModulationMatrixDashboardViewFactory>());
+	}
+}
+
+void FAudioModulationInsightsModule::ShutdownModule()
+{
+#if WITH_EDITOR
+	if (IAudioInsightsEditorModule::IsModuleLoaded())
+	{
+		return;
+	}
+
+	IAudioInsightsEditorModule& AudioInsightsModule = IAudioInsightsEditorModule::GetChecked();
+#else
+	if (IAudioInsightsModule::IsModuleLoaded())
+	{
+		return;
+	}
+
+	IAudioInsightsModule& AudioInsightsModule = IAudioInsightsModule::GetChecked();
+#endif // WITH_EDITOR
+		
+	AudioInsightsModule.UnregisterDashboardViewFactory("ControlBuses");
+	AudioInsightsModule.UnregisterDashboardViewFactory("ModulationMatrix");
+}
+
+IMPLEMENT_MODULE(FAudioModulationInsightsModule, AudioModulationInsights);
diff --git a/Engine/Plugins/Runtime/AudioModulationInsights/Source/AudioModulationInsights/Private/AudioModulationInsightsStyle.cpp b/Engine/Plugins/Runtime/AudioModulationInsights/Source/AudioModulationInsights/Private/AudioModulationInsightsStyle.cpp
new file mode 100644
--- /dev/null
+++ b/Engine/Plugins/Runtime/AudioModulationInsights/Source/AudioModulationInsights/Private/AudioModulationInsightsStyle.cpp
@@ -0,0 +1,58 @@
+// Copyright Epic Games, Inc. All Rights Reserved.
+
+#include "AudioModulationInsightsStyle.h"
+
+#include "Interfaces/IPluginManager.h"
+#include "Misc/Paths.h"
+#include "Styling/AppStyle.h"
+#include "Styling/SlateStyle.h"
+#include "Styling/SlateStyleMacros.h"
+#include "Styling/SlateStyleRegistry.h"
+#include "Textures/SlateIcon.h"
+
+#define LOCTEXT_NAMESPACE "AudioModulationInsights"
+
+FAudioModulationInsightsStyle::FAudioModulationInsightsStyle()
+	: FSlateStyleSet(GetStyleName())
+{
+	SetParentStyleName(FAppStyle::GetAppStyleSetName());
+
+	// Register style content paths
+	const TSharedPtr<const IPlugin> Plugin = IPluginManager::Get().FindPlugin(UE_PLUGIN_NAME);
+
+	SetContentRoot(Plugin->GetBaseDir() / "Resources");
+	SetCoreContentRoot(FPaths::EngineContentDir() / "Slate");
+
+	// Icons
+	const FVector2D Icon16(16.0f, 16.0f);
+
+	Set("ControlBus", new IMAGE_BRUSH(TEXT("control_bus"), Icon16));
+	Set("ControlBusMix", new IMAGE_BRUSH(TEXT("control_bus_mix"), Icon16));
+
+	// Register slate style
+	FSlateStyleRegistry::RegisterSlateStyle(*this);
+}
+
+FAudioModulationInsightsStyle::~FAudioModulationInsightsStyle()
+{
+	FSlateStyleRegistry::UnRegisterSlateStyle(*this);
+}
+
+const FAudioModulationInsightsStyle& FAudioModulationInsightsStyle::Get()
+{
+	static FAudioModulationInsightsStyle AudioModulationInsightsStyle;
+	return AudioModulationInsightsStyle;
+}
+
+const FName& FAudioModulationInsightsStyle::GetStyleName()
+{
+	static const FName StyleName = "AudioModulationInsightsStyle";
+	return StyleName;
+}
+
+FSlateIcon FAudioModulationInsightsStyle::CreateIcon(const FName& InName) const
+{
+	return { GetStyleName(), InName };
+}
+
+#undef LOCTEXT_NAMESPACE
diff --git a/Engine/Plugins/Runtime/AudioModulationInsights/Source/AudioModulationInsights/Private/AudioModulationInsightsStyle.h b/Engine/Plugins/Runtime/AudioModulationInsights/Source/AudioModulationInsights/Private/AudioModulationInsightsStyle.h
new file mode 100644
--- /dev/null
+++ b/Engine/Plugins/Runtime/AudioModulationInsights/Source/AudioModulationInsights/Private/AudioModulationInsightsStyle.h
@@ -0,0 +1,19 @@
+// Copyright Epic Games, Inc. All Rights Reserved.
+
+#pragma once
+
+#include "Styling/SlateStyle.h"
+#include "Textures/SlateIcon.h"
+#include "UObject/NameTypes.h"
+
+class FAudioModulationInsightsStyle final : public FSlateStyleSet
+{
+public:
+	FAudioModulationInsightsStyle();
+	~FAudioModulationInsightsStyle();
+	
+	static const FAudioModulationInsightsStyle& Get();
+	static const FName& GetStyleName();
+
+	FSlateIcon CreateIcon(const FName& InName) const;
+};
diff --git a/Engine/Plugins/Runtime/AudioModulationInsights/Source/AudioModulationInsights/Private/Messages/ControlBusTraceMessages.h b/Engine/Plugins/Runtime/AudioModulationInsights/Source/AudioModulationInsights/Private/Messages/ControlBusTraceMessages.h
new file mode 100644
--- /dev/null
+++ b/Engine/Plugins/Runtime/AudioModulationInsights/Source/AudioModulationInsights/Private/Messages/ControlBusTraceMessages.h
@@ -0,0 +1,88 @@
+// Copyright Epic Games, Inc. All Rights Reserved.
+#pragma once
+
+#include "AudioDefines.h"
+#include "Messages/AnalyzerMessageQueue.h"
+#include "Trace/Analyzer.h"
+#include "Views/TableDashboardViewFactory.h"
+
+namespace AudioModulationInsights
+{
+
+	using FBusId = uint32;
+
+	struct FControlBusMessageBase
+	{
+		FControlBusMessageBase() = default;
+		FControlBusMessageBase(const UE::Trace::IAnalyzer::FOnEventContext& InContext)
+		{
+			const UE::Trace::IAnalyzer::FEventData& EventData = InContext.EventData;
+			DeviceId = static_cast<::Audio::FDeviceId>(EventData.GetValue<uint32>("DeviceId"));
+			Timestamp = InContext.EventTime.AsSeconds(EventData.GetValue<uint64>("Timestamp"));
+			ControlBusId = static_cast<FBusId>(EventData.GetValue<uint32>("ControlBusId"));
+		}
+
+		Audio::FDeviceId DeviceId = INDEX_NONE;
+		FBusId ControlBusId = INDEX_NONE;
+		double Timestamp = 0.0;
+	};
+
+	struct FControlBusActivateMessage : public FControlBusMessageBase
+	{
+		FControlBusActivateMessage() = default;
+		FControlBusActivateMessage(const UE::Trace::IAnalyzer::FOnEventContext& InContext)
+			: FControlBusMessageBase(InContext)
+		{
+			const UE::Trace::IAnalyzer::FEventData& EventData = InContext.EventData;
+			EventData.GetString("Name", BusName);
+			EventData.GetString("ParamName", ParamName);
+		}
+
+		FString BusName;
+		FString ParamName;
+	};
+
+	using FControlBusDeactivateMessage = FControlBusMessageBase;
+
+	struct FControlBusUpdateMessage : public FControlBusMessageBase
+	{
+		FControlBusUpdateMessage() = default;
+		FControlBusUpdateMessage(const UE::Trace::IAnalyzer::FOnEventContext& InContext)
+			: FControlBusMessageBase(InContext)
+		{
+			const UE::Trace::IAnalyzer::FEventData& EventData = InContext.EventData;
+			Value = EventData.GetValue<float>("Value");
+			EventData.GetString("Name", BusName);
+			EventData.GetString("ParamName", ParamName);
+		}
+
+		FString BusName;
+		FString ParamName;
+		float Value = 1.0f;
+	};
+
+	class FControlBusDashboardEntry : public UE::Audio::Insights::FSoundAssetDashboardEntry
+	{
+	public:
+		FControlBusDashboardEntry() = default;
+		virtual ~FControlBusDashboardEntry() = default;
+
+		FText GetDisplayNameAsFText() const { return FText::FromString(DisplayName); }
+		FText GetParamNameAsFText() const { return FText::FromString(ParamName); }
+
+		FBusId ControlBusId = INDEX_NONE;
+		float Value = 1.0f;
+		FString DisplayName;
+		FString ParamName;
+	};
+
+	class FControlBusMessages
+	{
+		UE::Audio::Insights::TAnalyzerMessageQueue<FControlBusUpdateMessage> UpdateMessages{ 2.0 };
+		UE::Audio::Insights::TAnalyzerMessageQueue<FControlBusActivateMessage> ActivateMessages{ 0.1 };
+		UE::Audio::Insights::TAnalyzerMessageQueue<FControlBusDeactivateMessage> DeactivateMessages{ 0.1 };
+
+		friend class FControlBusTraceProvider;
+	};
+
+} // namespace AudioModulationInsights
diff --git a/Engine/Plugins/Runtime/AudioModulationInsights/Source/AudioModulationInsights/Private/Messages/ModulationMatrixTraceMessages.h b/Engine/Plugins/Runtime/AudioModulationInsights/Source/AudioModulationInsights/Private/Messages/ModulationMatrixTraceMessages.h
new file mode 100644
--- /dev/null
+++ b/Engine/Plugins/Runtime/AudioModulationInsights/Source/AudioModulationInsights/Private/Messages/ModulationMatrixTraceMessages.h
@@ -0,0 +1,191 @@
+// Copyright Epic Games, Inc. All Rights Reserved.
+#pragma once
+
+#include "AudioDefines.h"
+#include "Messages/AnalyzerMessageQueue.h"
+#include "Trace/Analyzer.h"
+#include "Views/TableDashboardViewFactory.h"
+
+namespace AudioModulationInsights
+{
+	using FBusId    = uint32;
+	using FSourceId = uint32;
+
+	enum class EModulationMatrixEntryType : uint8
+	{
+		// Arranged alphabetically to simplify sorting
+		BusMix,
+		BusFinalValues,
+		Generator,
+		None
+	};
+
+
+	// Modulation matrix dashboard entry
+	class FModulationMatrixDashboardEntry : public UE::Audio::Insights::FSoundAssetDashboardEntry
+	{
+	public:
+		FModulationMatrixDashboardEntry() = default;
+		virtual ~FModulationMatrixDashboardEntry() = default;
+
+		FSourceId SourceId = INDEX_NONE;
+		EModulationMatrixEntryType EntryType = EModulationMatrixEntryType::None;
+		FString DisplayName;
+
+		TMap<FBusId, float> BusIdToValueMap;
+	};
+	
+
+	// Modulation matrix message base
+	struct FModulationMatrixMessageBase
+	{
+		FModulationMatrixMessageBase() = default;
+
+		FModulationMatrixMessageBase(const UE::Trace::IAnalyzer::FOnEventContext& InContext)
+		{
+			const UE::Trace::IAnalyzer::FEventData& EventData = InContext.EventData;
+
+			DeviceId  = static_cast<::Audio::FDeviceId>(EventData.GetValue<uint32>("DeviceId"));
+			Timestamp = InContext.EventTime.AsSeconds(EventData.GetValue<uint64>("Timestamp"));
+			SourceId  = static_cast<FSourceId>(EventData.GetValue<uint32>("SourceId"));
+		}
+
+		Audio::FDeviceId DeviceId = INDEX_NONE;
+		FSourceId SourceId = INDEX_NONE;
+		double Timestamp = 0.0;
+	};
+
+
+	// Register bus message
+	struct FModulationMatrixRegisterBusMessage : public FModulationMatrixMessageBase
+	{
+		FModulationMatrixRegisterBusMessage() = default;
+
+		FModulationMatrixRegisterBusMessage(const UE::Trace::IAnalyzer::FOnEventContext& InContext)
+			: FModulationMatrixMessageBase(InContext)
+		{
+			const UE::Trace::IAnalyzer::FEventData& EventData = InContext.EventData;
+
+			ModulatingSourceId = static_cast<FSourceId>(EventData.GetValue<uint32>("ModulatingSourceId"));
+			EventData.GetString("BusName", BusName);
+		}
+
+		FSourceId ModulatingSourceId = INDEX_NONE;
+		FString BusName;
+	};
+
+
+	// Activate messages
+	struct FModulationMatrixActivateMessage : public FModulationMatrixMessageBase
+	{
+		FModulationMatrixActivateMessage() = default;
+
+		FModulationMatrixActivateMessage(const UE::Trace::IAnalyzer::FOnEventContext& InContext)
+			: FModulationMatrixMessageBase(InContext)
+		{
+			const UE::Trace::IAnalyzer::FEventData& EventData = InContext.EventData;
+
+			EventData.GetString("Name", Name);
+		}
+
+		EModulationMatrixEntryType EntryType = EModulationMatrixEntryType::None;
+		FString Name;
+	};
+
+	struct FBusMixActivateMessage : public FModulationMatrixActivateMessage
+	{
+		FBusMixActivateMessage() = default;
+
+		FBusMixActivateMessage(const UE::Trace::IAnalyzer::FOnEventContext& InContext)
+			: FModulationMatrixActivateMessage(InContext)
+		{
+			EntryType = EModulationMatrixEntryType::BusMix;
+		}
+	};
+
+	struct FGeneratorActivateMessage : public FModulationMatrixActivateMessage
+	{
+		FGeneratorActivateMessage() = default;
+
+		FGeneratorActivateMessage(const UE::Trace::IAnalyzer::FOnEventContext& InContext)
+			: FModulationMatrixActivateMessage(InContext)
+		{
+			EntryType = EModulationMatrixEntryType::Generator;
+		}
+	};
+
+
+	// Update messages
+	struct FModulationMatrixUpdateMessage : public FModulationMatrixMessageBase
+	{
+		FModulationMatrixUpdateMessage() = default;
+
+		FModulationMatrixUpdateMessage(const UE::Trace::IAnalyzer::FOnEventContext& InContext)
+			: FModulationMatrixMessageBase(InContext)
+		{
+			const UE::Trace::IAnalyzer::FEventData& EventData = InContext.EventData;
+
+			const TArrayView<const FBusId> BusIds   = EventData.GetArrayView<FBusId>("BusIds");
+			const TArrayView<const float> BusValues = EventData.GetArrayView<float>("BusValues");
+
+			check(BusIds.Num() == BusValues.Num());
+
+			for (int32 Index = 0; Index < BusIds.Num(); ++Index)
+			{
+				BusIdToValueMap.Emplace(BusIds[Index], BusValues[Index]);
+			}
+		}
+
+		TMap<FBusId, float> BusIdToValueMap;
+	};
+
+	using FBusMixUpdateMessage = FModulationMatrixUpdateMessage;
+
+	struct FGeneratorUpdateMessage : public FModulationMatrixMessageBase
+	{
+		FGeneratorUpdateMessage() = default;
+
+		FGeneratorUpdateMessage(const UE::Trace::IAnalyzer::FOnEventContext& InContext)
+			: FModulationMatrixMessageBase(InContext)
+		{
+			const UE::Trace::IAnalyzer::FEventData& EventData = InContext.EventData;
+
+			const TArrayView<const FBusId> BusIds = EventData.GetArrayView<FBusId>("BusIds");
+			const float GeneratorValue = EventData.GetValue<float>("GeneratorValue");
+
+			// All control buses use the same generator value
+			for (int32 Index = 0; Index < BusIds.Num(); ++Index)
+			{
+				BusIdToValueMap.Emplace(BusIds[Index], GeneratorValue);
+			}
+		}
+
+		TMap<FBusId, float> BusIdToValueMap;
+	};
+
+	using FBusFinalValuesUpdateMessage = FModulationMatrixUpdateMessage;
+
+
+	// Deactivate message
+	using FModulationMatrixDeactivateMessage = FModulationMatrixMessageBase;
+
+
+	// Modulation matrix messages
+	class FModulationMatrixMessages
+	{
+		UE::Audio::Insights::TAnalyzerMessageQueue<FModulationMatrixRegisterBusMessage> RegisterBusMessages{ 0.1 };
+
+		UE::Audio::Insights::TAnalyzerMessageQueue<FBusMixActivateMessage> BusMixActivateMessages{ 0.1 };
+		UE::Audio::Insights::TAnalyzerMessageQueue<FBusMixUpdateMessage> BusMixUpdateMessages{ 0.1 };
+
+		UE::Audio::Insights::TAnalyzerMessageQueue<FGeneratorActivateMessage> GeneratorActivateMessages{ 0.1 };
+		UE::Audio::Insights::TAnalyzerMessageQueue<FGeneratorUpdateMessage> GeneratorUpdateMessages{ 0.1 };
+
+		UE::Audio::Insights::TAnalyzerMessageQueue<FBusFinalValuesUpdateMessage> BusFinalValuesUpdateMessages{ 0.1 };
+
+		UE::Audio::Insights::TAnalyzerMessageQueue<FModulationMatrixDeactivateMessage> DeactivateMessages{ 0.1 };
+
+		friend class FModulationMatrixTraceProvider;
+	};
+
+} // namespace AudioModulationInsights
diff --git a/Engine/Plugins/Runtime/AudioModulationInsights/Source/AudioModulationInsights/Private/Providers/ControlBusTraceProvider.cpp b/Engine/Plugins/Runtime/AudioModulationInsights/Source/AudioModulationInsights/Private/Providers/ControlBusTraceProvider.cpp
new file mode 100644
--- /dev/null
+++ b/Engine/Plugins/Runtime/AudioModulationInsights/Source/AudioModulationInsights/Private/Providers/ControlBusTraceProvider.cpp
@@ -0,0 +1,148 @@
+// Copyright Epic Games, Inc. All Rights Reserved.
+#include "ControlBusTraceProvider.h"
+
+#include "AudioInsightsUtils.h"
+#include "Trace/Analyzer.h"
+#include "TraceServices/Model/AnalysisSession.h"
+
+namespace AudioModulationInsights
+{
+	FName FControlBusTraceProvider::GetName_Static()
+	{
+		return "ControlBusProvider";
+	}
+
+	bool FControlBusTraceProvider::ProcessMessages()
+	{
+		using namespace UE::Audio::Insights;
+
+		auto BumpEntryFunc = [this](const FControlBusMessageBase& Msg)
+		{
+			TSharedPtr<FControlBusDashboardEntry>* ToReturn = nullptr;
+			UpdateDeviceEntry(Msg.DeviceId, Msg.ControlBusId, [&ToReturn, &Msg](TSharedPtr<FControlBusDashboardEntry>& Entry)
+				{
+					if (!Entry.IsValid())
+					{
+						Entry = MakeShared<FControlBusDashboardEntry>();
+						Entry->DeviceId = Msg.DeviceId;
+						Entry->ControlBusId = Msg.ControlBusId;
+					}
+					ToReturn = &Entry;
+				});
+
+			return ToReturn;
+		};
+
+		ProcessMessageQueue<FControlBusActivateMessage>(TraceMessages.ActivateMessages, BumpEntryFunc,
+		[](const FControlBusActivateMessage& Msg, TSharedPtr<FControlBusDashboardEntry>* OutEntry)
+		{
+			FControlBusDashboardEntry& EntryRef = *OutEntry->Get();
+			EntryRef.Name = *Msg.BusName;
+			EntryRef.ControlBusId = Msg.ControlBusId;
+			EntryRef.DisplayName = AudioInsightsUtils::ResolveObjectDisplayName(EntryRef.Name);
+			EntryRef.ParamName = Msg.ParamName;
+		});
+
+		ProcessMessageQueue<FControlBusUpdateMessage>(TraceMessages.UpdateMessages, BumpEntryFunc,
+		[](const FControlBusUpdateMessage& Msg, TSharedPtr<FControlBusDashboardEntry>* OutEntry)
+		{
+			FControlBusDashboardEntry& EntryRef = *OutEntry->Get();
+			EntryRef.Name = *Msg.BusName;
+			EntryRef.DisplayName = AudioInsightsUtils::ResolveObjectDisplayName(EntryRef.Name);
+			EntryRef.ParamName = Msg.ParamName;
+			EntryRef.Value = Msg.Value;
+		});
+
+		auto GetEntry = [this](const FControlBusMessageBase& Msg)
+		{
+			return FindDeviceEntry(Msg.DeviceId, Msg.ControlBusId);
+		};
+
+		ProcessMessageQueue<FControlBusDeactivateMessage>(TraceMessages.DeactivateMessages, GetEntry,
+		[this](const FControlBusDeactivateMessage& Msg, TSharedPtr<FControlBusDashboardEntry>* OutEntry)
+		{
+			if (OutEntry && (*OutEntry)->Timestamp < Msg.Timestamp)
+			{
+				RemoveDeviceEntry(Msg.DeviceId, Msg.ControlBusId);
+			}
+		});
+
+		return true;
+	}
+
+	UE::Trace::IAnalyzer* FControlBusTraceProvider::ConstructAnalyzer(TraceServices::IAnalysisSession& InSession)
+	{
+		class FControlBusTraceAnalyzer : public UE::Audio::Insights::FTraceProviderBase::FTraceAnalyzerBase
+		{
+		public:
+			FControlBusTraceAnalyzer(TSharedRef<FControlBusTraceProvider> InProvider, TraceServices::IAnalysisSession& InSession)
+				: UE::Audio::Insights::FTraceProviderBase::FTraceAnalyzerBase(InProvider)
+				, Session(InSession)
+			{
+			}
+
+			virtual void OnAnalysisBegin(const UE::Trace::IAnalyzer::FOnAnalysisContext& Context) override
+			{
+				UE::Audio::Insights::FTraceProviderBase::FTraceAnalyzerBase::OnAnalysisBegin(Context);
+
+				UE::Trace::IAnalyzer::FInterfaceBuilder& Builder = Context.InterfaceBuilder;
+				Builder.RouteEvent(RouteId_Activate, "Audio", "ControlBusActivate");
+				Builder.RouteEvent(RouteId_Deactivate, "Audio", "ControlBusDeactivate");
+				Builder.RouteEvent(RouteId_Update, "Audio", "ControlBusUpdate");
+			}
+
+			virtual bool OnHandleEvent(uint16 RouteId, UE::Trace::IAnalyzer::EStyle Style, const UE::Trace::IAnalyzer::FOnEventContext& Context) override
+			{
+				LLM_SCOPE_BYNAME(TEXT("Insights/FControlBusTraceAnalyzer"));
+
+				FControlBusMessages& Messages = GetProvider<FControlBusTraceProvider>().TraceMessages;
+				switch (RouteId)
+				{
+					case RouteId_Activate:
+					{
+						Messages.ActivateMessages.Enqueue(FControlBusActivateMessage { Context });
+						break;
+					}
+
+					case RouteId_Deactivate:
+					{
+						Messages.DeactivateMessages.Enqueue(FControlBusDeactivateMessage{ Context });
+						break;
+					}
+
+					case RouteId_Update:
+					{
+						Messages.UpdateMessages.Enqueue(FControlBusUpdateMessage { Context });
+						break;
+					}
+
+					default:
+					{
+						return OnEventFailure(RouteId, Style, Context);
+					}
+				}
+
+				const double Timestamp = Context.EventTime.AsSeconds(Context.EventData.GetValue<uint64>("Timestamp"));
+
+				{
+					TraceServices::FAnalysisSessionEditScope SessionEditScope(Session);
+					Session.UpdateDurationSeconds(Timestamp);
+				}
+
+				return OnEventSuccess(RouteId, Style, Context);
+			}
+
+		private:
+			enum : uint16
+			{
+				RouteId_Activate,
+				RouteId_Deactivate,
+				RouteId_Update
+			};
+
+			TraceServices::IAnalysisSession& Session;
+		};
+
+		return new FControlBusTraceAnalyzer(AsShared(), InSession);
+	}
+} // namespace AudioModulationInsights
diff --git a/Engine/Plugins/Runtime/AudioModulationInsights/Source/AudioModulationInsights/Private/Providers/ControlBusTraceProvider.h b/Engine/Plugins/Runtime/AudioModulationInsights/Source/AudioModulationInsights/Private/Providers/ControlBusTraceProvider.h
new file mode 100644
--- /dev/null
+++ b/Engine/Plugins/Runtime/AudioModulationInsights/Source/AudioModulationInsights/Private/Providers/ControlBusTraceProvider.h
@@ -0,0 +1,31 @@
+// Copyright Epic Games, Inc. All Rights Reserved.
+#pragma once
+
+#include "AudioInsightsTraceProviderBase.h"
+#include "Messages/ControlBusTraceMessages.h"
+
+
+namespace AudioModulationInsights
+{
+	class FControlBusTraceProvider
+		: public UE::Audio::Insights::TDeviceDataMapTraceProvider<uint32, TSharedPtr<FControlBusDashboardEntry>>
+		, public TSharedFromThis<FControlBusTraceProvider>
+	{
+	public:
+		FControlBusTraceProvider()
+			: UE::Audio::Insights::TDeviceDataMapTraceProvider<uint32, TSharedPtr<FControlBusDashboardEntry>>(GetName_Static())
+		{
+
+		}
+
+		virtual ~FControlBusTraceProvider() = default;
+		virtual UE::Trace::IAnalyzer* ConstructAnalyzer(TraceServices::IAnalysisSession& InSession) override;
+
+		virtual bool ProcessMessages() override;
+		static FName GetName_Static();
+
+	private:
+
+		FControlBusMessages TraceMessages;
+	};
+} // namespace AudioModulationInsights
diff --git a/Engine/Plugins/Runtime/AudioModulationInsights/Source/AudioModulationInsights/Private/Providers/ModulationMatrixTraceProvider.cpp b/Engine/Plugins/Runtime/AudioModulationInsights/Source/AudioModulationInsights/Private/Providers/ModulationMatrixTraceProvider.cpp
new file mode 100644
--- /dev/null
+++ b/Engine/Plugins/Runtime/AudioModulationInsights/Source/AudioModulationInsights/Private/Providers/ModulationMatrixTraceProvider.cpp
@@ -0,0 +1,367 @@
+// Copyright Epic Games, Inc. All Rights Reserved.
+#include "ModulationMatrixTraceProvider.h"
+
+#include "AudioInsightsUtils.h"
+#include "Trace/Analyzer.h"
+#include "TraceServices/Model/AnalysisSession.h"
+
+#if WITH_EDITOR
+#include "AudioDeviceManager.h"
+#endif // WITH_EDITOR
+
+namespace AudioModulationInsights
+{
+	FModulationMatrixTraceProvider::FModulationMatrixTraceProvider()
+		: UE::Audio::Insights::TDeviceDataMapTraceProvider<uint32, TSharedPtr<FModulationMatrixDashboardEntry>>(GetName_Static())
+	{
+#if WITH_EDITOR
+		FAudioDeviceManagerDelegates::OnAudioDeviceDestroyed.AddRaw(this, &FModulationMatrixTraceProvider::OnAudioDeviceDestroyed);
+#endif // WITH_EDITOR
+	}
+
+	FModulationMatrixTraceProvider::~FModulationMatrixTraceProvider()
+	{
+#if WITH_EDITOR
+		FAudioDeviceManagerDelegates::OnAudioDeviceDestroyed.RemoveAll(this);
+#endif // WITH_EDITOR
+	}
+
+	FName FModulationMatrixTraceProvider::GetName_Static()
+	{
+		static const FName ModulationMatrixProviderName = "ModulationMatrixProvider";
+		return ModulationMatrixProviderName;
+	}
+
+	void FModulationMatrixTraceProvider::UpdateActiveControlBusesToAdd(const TMap<FBusId, FString>& InBusIdToBusNameMap)
+	{
+		for (const auto& [BusId, BusName] : InBusIdToBusNameMap)
+		{
+			if (BusInfo* FoundBusInfoPtr = ActiveControlBuses.Find(BusId))
+			{
+				BusInfo& FoundBusInfo = *FoundBusInfoPtr;
+				FoundBusInfo.RefCount++;
+			}
+			else
+			{
+				ActiveControlBuses.Emplace(BusId, { BusName, 1 });
+			}
+		}
+
+		ActiveControlBuses.ValueSort([](const BusInfo& A, const BusInfo& B)
+		{
+			return A.BusName.ToLower() < B.BusName.ToLower();
+		});
+	}
+
+	void FModulationMatrixTraceProvider::UpdateActiveControlBusesToRemove(const TMap<FBusId, float>& InBusIdToValueMap)
+	{
+		TArray<FBusId> BusesPendingToRemove;
+
+		for (const auto& [BusId, BusValue] : InBusIdToValueMap)
+		{
+			if (BusInfo* FoundBusInfoPtr = ActiveControlBuses.Find(BusId))
+			{
+				BusInfo& FoundBusInfo = *FoundBusInfoPtr;
+				FoundBusInfo.RefCount--;
+
+				if (FoundBusInfo.RefCount <= 0)
+				{
+					RemovedControlBusesNames.Emplace(FoundBusInfo.BusName);
+					BusesPendingToRemove.Emplace(BusId);
+				}
+			}
+		}
+
+		for (const FBusId& BusId : BusesPendingToRemove)
+		{
+			ActiveControlBuses.Remove(BusId);
+		}
+	}
+
+	bool FModulationMatrixTraceProvider::ProcessMessages()
+	{
+		using namespace UE::Audio::Insights;
+
+		// Helper lambdas
+		auto BumpEntryFunc = [this](const FModulationMatrixMessageBase& Msg)
+		{
+			TSharedPtr<FModulationMatrixDashboardEntry>* ToReturn = nullptr;
+
+			UpdateDeviceEntry(Msg.DeviceId, Msg.SourceId, [&ToReturn, &Msg](TSharedPtr<FModulationMatrixDashboardEntry>& Entry)
+			{
+				if (!Entry.IsValid())
+				{
+					Entry = MakeShared<FModulationMatrixDashboardEntry>();
+					Entry->DeviceId = Msg.DeviceId;
+					Entry->SourceId = Msg.SourceId;
+				}
+
+				Entry->Timestamp = Msg.Timestamp;
+
+				ToReturn = &Entry;
+			});
+
+			return ToReturn;
+		};
+
+		auto GetEntry = [this](const FModulationMatrixMessageBase& Msg)
+		{
+			return FindDeviceEntry(Msg.DeviceId, Msg.SourceId);
+		};
+
+		// Process messages
+		uint32 NumActiveBuses = ActiveControlBuses.Num();
+
+		if (!(TraceMessages.RegisterBusMessages.IsEmpty() && (TraceMessages.BusMixActivateMessages.IsEmpty() || TraceMessages.GeneratorActivateMessages.IsEmpty())))
+		{
+			TMap<FSourceId, TMap<FBusId, FString>> ModulatingSourceIdToBusesInfoMap;
+
+			auto UpdateActiveModulatorSourceIds = [&](const Audio::FDeviceId DeviceId, const FSourceId SourceId)
+			{
+				TSet<uint32>& ActiveModulatorSourceIds = DeviceIdToActiveModulatorSourceIdsMap.FindOrAdd(DeviceId);
+
+				if (!ActiveModulatorSourceIds.Contains(SourceId))
+				{
+					ActiveModulatorSourceIds.Add(SourceId);
+
+					if (const TMap<FBusId, FString>* FoundBusInfoPtr = ModulatingSourceIdToBusesInfoMap.Find(SourceId))
+					{
+						UpdateActiveControlBusesToAdd(*FoundBusInfoPtr);
+					}
+				}
+			};
+
+			ProcessMessageQueue<FModulationMatrixRegisterBusMessage>(TraceMessages.RegisterBusMessages,
+			[](const FModulationMatrixMessageBase& Msg)	{ return nullptr; }, // No entry operation
+			[this, &ModulatingSourceIdToBusesInfoMap](const FModulationMatrixRegisterBusMessage& Msg, TSharedPtr<FModulationMatrixDashboardEntry>* OutEntry)
+			{
+				if (TMap<FBusId, FString>* FoundBusInfoPtr = ModulatingSourceIdToBusesInfoMap.Find(Msg.ModulatingSourceId))
+				{
+					FoundBusInfoPtr->FindOrAdd(Msg.SourceId, AudioInsightsUtils::ResolveObjectDisplayName(Msg.BusName));
+				}
+				else
+				{
+					ModulatingSourceIdToBusesInfoMap.Emplace(Msg.ModulatingSourceId).Emplace(Msg.SourceId, AudioInsightsUtils::ResolveObjectDisplayName(Msg.BusName));
+				}
+			});
+
+			ProcessMessageQueue<FBusMixActivateMessage>(TraceMessages.BusMixActivateMessages, BumpEntryFunc,
+			[this, &UpdateActiveModulatorSourceIds](const FBusMixActivateMessage& Msg, TSharedPtr<FModulationMatrixDashboardEntry>* OutEntry)
+			{
+				FModulationMatrixDashboardEntry& EntryRef = *OutEntry->Get();
+				EntryRef.Name = *Msg.Name;
+				EntryRef.DisplayName = AudioInsightsUtils::ResolveObjectDisplayName(Msg.Name);
+				EntryRef.EntryType = EModulationMatrixEntryType::BusMix;
+
+				UpdateActiveModulatorSourceIds(Msg.DeviceId, EntryRef.SourceId);
+			});
+
+			ProcessMessageQueue<FGeneratorActivateMessage>(TraceMessages.GeneratorActivateMessages, BumpEntryFunc,
+			[this, &UpdateActiveModulatorSourceIds](const FGeneratorActivateMessage& Msg, TSharedPtr<FModulationMatrixDashboardEntry>* OutEntry)
+			{
+				FModulationMatrixDashboardEntry& EntryRef = *OutEntry->Get();
+				EntryRef.Name = *Msg.Name;
+				EntryRef.DisplayName = EntryRef.GetDisplayName().ToString();
+				EntryRef.EntryType = EModulationMatrixEntryType::Generator;
+
+				UpdateActiveModulatorSourceIds(Msg.DeviceId, EntryRef.SourceId);
+			});
+		}
+
+		ProcessMessageQueue<FModulationMatrixDeactivateMessage>(TraceMessages.DeactivateMessages, GetEntry,
+		[this](const FModulationMatrixDeactivateMessage& Msg, TSharedPtr<FModulationMatrixDashboardEntry>* OutEntry)
+		{
+			if (OutEntry && (*OutEntry)->Timestamp < Msg.Timestamp)
+			{
+				const FModulationMatrixDashboardEntry& EntryRef = *OutEntry->Get();
+				UpdateActiveControlBusesToRemove(EntryRef.BusIdToValueMap);
+
+				if (DeviceIdToActiveModulatorSourceIdsMap.Contains(Msg.DeviceId))
+				{
+					DeviceIdToActiveModulatorSourceIdsMap[Msg.DeviceId].Remove(EntryRef.SourceId);
+				}
+
+				RemoveDeviceEntry(Msg.DeviceId, Msg.SourceId);
+			}
+		});
+
+		ProcessMessageQueue<FBusMixUpdateMessage>(TraceMessages.BusMixUpdateMessages, GetEntry,
+		[this](const FBusMixUpdateMessage& Msg, TSharedPtr<FModulationMatrixDashboardEntry>* OutEntry)
+		{
+			if (OutEntry)
+			{
+				FModulationMatrixDashboardEntry& EntryRef = *OutEntry->Get();
+				EntryRef.BusIdToValueMap = Msg.BusIdToValueMap;
+				EntryRef.Timestamp = Msg.Timestamp;
+			}
+		});
+
+		ProcessMessageQueue<FGeneratorUpdateMessage>(TraceMessages.GeneratorUpdateMessages, GetEntry,
+		[this](const FGeneratorUpdateMessage& Msg, TSharedPtr<FModulationMatrixDashboardEntry>* OutEntry)
+		{
+			if (OutEntry)
+			{
+				FModulationMatrixDashboardEntry& EntryRef = *OutEntry->Get();
+				EntryRef.BusIdToValueMap = Msg.BusIdToValueMap;
+				EntryRef.Timestamp = Msg.Timestamp;
+			}
+		});
+
+		ProcessMessageQueue<FBusFinalValuesUpdateMessage>(TraceMessages.BusFinalValuesUpdateMessages, BumpEntryFunc,
+		[this](const FBusFinalValuesUpdateMessage& Msg, TSharedPtr<FModulationMatrixDashboardEntry>* OutEntry)
+		{
+			FModulationMatrixDashboardEntry& EntryRef = *OutEntry->Get();
+
+			if (EntryRef.Name.IsEmpty())
+			{
+				static const FString BusFinalValuesEntryName("Final Values: ");
+				EntryRef.Name = BusFinalValuesEntryName;
+				EntryRef.DisplayName = BusFinalValuesEntryName;
+				EntryRef.EntryType = EModulationMatrixEntryType::BusFinalValues;
+			}
+
+			EntryRef.BusIdToValueMap = Msg.BusIdToValueMap;
+		});
+
+		// Notify Control Buses update
+		if (NumActiveBuses != ActiveControlBuses.Num() && !ActiveControlBuses.IsEmpty())
+		{
+			OnControlBusesAdded.ExecuteIfBound(ActiveControlBuses);
+		}
+
+		if (!RemovedControlBusesNames.IsEmpty())
+		{
+			OnControlBusesRemoved.ExecuteIfBound(RemovedControlBusesNames);
+			RemovedControlBusesNames.Reset();
+		}
+
+		return true;
+	}
+
+#if WITH_EDITOR
+	void FModulationMatrixTraceProvider::OnAudioDeviceDestroyed(::Audio::FDeviceId InDeviceId)
+	{
+		DeviceIdToActiveModulatorSourceIdsMap.Reset();
+		ActiveControlBuses.Reset();
+	}
+#endif // WITH_EDITOR
+
+	UE::Trace::IAnalyzer* FModulationMatrixTraceProvider::ConstructAnalyzer(TraceServices::IAnalysisSession& InSession)
+	{
+		class FModulationMatrixTraceAnalyzer : public UE::Audio::Insights::FTraceProviderBase::FTraceAnalyzerBase
+		{
+		public:
+			FModulationMatrixTraceAnalyzer(TSharedRef<FModulationMatrixTraceProvider> InProvider, TraceServices::IAnalysisSession& InSession)
+				: UE::Audio::Insights::FTraceProviderBase::FTraceAnalyzerBase(InProvider)
+				, Session(InSession)
+			{
+			}
+
+			virtual void OnAnalysisBegin(const UE::Trace::IAnalyzer::FOnAnalysisContext& Context) override
+			{
+				UE::Audio::Insights::FTraceProviderBase::FTraceAnalyzerBase::OnAnalysisBegin(Context);
+
+				UE::Trace::IAnalyzer::FInterfaceBuilder& Builder = Context.InterfaceBuilder;
+
+				Builder.RouteEvent(RouteId_BusMixRegisterBus,    "Audio", "BusMixRegisterBus");
+				Builder.RouteEvent(RouteId_GeneratorRegisterBus, "Audio", "GeneratorRegisterBus");
+
+				Builder.RouteEvent(RouteId_BusMixActivate,       "Audio", "BusMixActivate");
+				Builder.RouteEvent(RouteId_GeneratorActivate,    "Audio", "GeneratorActivate");
+
+				Builder.RouteEvent(RouteId_BusMixUpdate,         "Audio", "BusMixUpdate");
+				Builder.RouteEvent(RouteId_GeneratorUpdate,      "Audio", "GeneratorUpdate");
+				Builder.RouteEvent(RouteId_BusFinalValuesUpdate, "Audio", "BusFinalValuesUpdate");
+
+				Builder.RouteEvent(RouteId_Deactivate, "Audio", "ModulatingSourceDeactivate");
+			}
+
+			virtual bool OnHandleEvent(uint16 RouteId, UE::Trace::IAnalyzer::EStyle Style, const UE::Trace::IAnalyzer::FOnEventContext& Context) override
+			{
+				LLM_SCOPE_BYNAME(TEXT("Insights/FModulationMatrixTraceAnalyzer"));
+
+				FModulationMatrixMessages& Messages = GetProvider<FModulationMatrixTraceProvider>().TraceMessages;
+
+				switch (RouteId)
+				{
+					case RouteId_BusMixRegisterBus:
+					case RouteId_GeneratorRegisterBus:
+					{
+						Messages.RegisterBusMessages.Enqueue(FModulationMatrixRegisterBusMessage{ Context });
+						break;
+					}
+
+					case RouteId_BusMixActivate:
+					{
+						Messages.BusMixActivateMessages.Enqueue(FBusMixActivateMessage{ Context });
+						break;
+					}
+
+					case RouteId_GeneratorActivate:
+					{
+						Messages.GeneratorActivateMessages.Enqueue(FGeneratorActivateMessage{ Context });
+						break;
+					}
+
+					case RouteId_BusMixUpdate:
+					{
+						Messages.BusMixUpdateMessages.Enqueue(FBusMixUpdateMessage{ Context });
+						break;
+					}
+
+					case RouteId_GeneratorUpdate:
+					{
+						Messages.GeneratorUpdateMessages.Enqueue(FGeneratorUpdateMessage{ Context });
+						break;
+					}
+
+					case RouteId_BusFinalValuesUpdate:
+					{
+						Messages.BusFinalValuesUpdateMessages.Enqueue(FBusFinalValuesUpdateMessage{ Context });
+						break;
+					}
+
+					case RouteId_Deactivate:
+					{
+						Messages.DeactivateMessages.Enqueue(FModulationMatrixDeactivateMessage{ Context });
+						break;
+					}
+
+					default:
+					{
+						return OnEventFailure(RouteId, Style, Context);
+					}
+				}
+
+				const double Timestamp = Context.EventTime.AsSeconds(Context.EventData.GetValue<uint64>("Timestamp"));
+
+				{
+					TraceServices::FAnalysisSessionEditScope SessionEditScope(Session);
+					Session.UpdateDurationSeconds(Timestamp);
+				}
+
+				return OnEventSuccess(RouteId, Style, Context);
+			}
+
+		private:
+			enum : uint16
+			{
+				RouteId_BusMixRegisterBus,
+				RouteId_GeneratorRegisterBus,
+
+				RouteId_BusMixActivate,
+				RouteId_GeneratorActivate,
+
+				RouteId_BusMixUpdate,
+				RouteId_GeneratorUpdate,
+				RouteId_BusFinalValuesUpdate,
+
+				RouteId_Deactivate
+			};
+
+			TraceServices::IAnalysisSession& Session;
+		};
+
+		return new FModulationMatrixTraceAnalyzer(AsShared(), InSession);
+	}
+} // namespace AudioModulationInsights
diff --git a/Engine/Plugins/Runtime/AudioModulationInsights/Source/AudioModulationInsights/Private/Providers/ModulationMatrixTraceProvider.h b/Engine/Plugins/Runtime/AudioModulationInsights/Source/AudioModulationInsights/Private/Providers/ModulationMatrixTraceProvider.h
new file mode 100644
--- /dev/null
+++ b/Engine/Plugins/Runtime/AudioModulationInsights/Source/AudioModulationInsights/Private/Providers/ModulationMatrixTraceProvider.h
@@ -0,0 +1,53 @@
+// Copyright Epic Games, Inc. All Rights Reserved.
+#pragma once
+
+#include "AudioInsightsTraceProviderBase.h"
+#include "Containers/Map.h"
+#include "Delegates/DelegateCombinations.h"
+#include "Messages/ModulationMatrixTraceMessages.h"
+
+namespace AudioModulationInsights
+{
+	class FModulationMatrixTraceProvider
+		: public UE::Audio::Insights::TDeviceDataMapTraceProvider<uint32, TSharedPtr<FModulationMatrixDashboardEntry>>
+		, public TSharedFromThis<FModulationMatrixTraceProvider>
+	{
+	public:
+		FModulationMatrixTraceProvider();
+
+		virtual ~FModulationMatrixTraceProvider();
+
+		static FName GetName_Static();
+
+		virtual bool ProcessMessages() override;
+		virtual UE::Trace::IAnalyzer* ConstructAnalyzer(TraceServices::IAnalysisSession& InSession) override;
+
+		struct BusInfo
+		{
+			FString BusName;
+			int32 RefCount = 0;
+		};
+
+		using BusIdToBusInfoMap = TMap<FBusId, BusInfo>;
+
+		DECLARE_DELEGATE_OneParam(FOnControlBusesAdded, const BusIdToBusInfoMap& /*AddedControlBuses*/);
+		FOnControlBusesAdded OnControlBusesAdded;
+
+		DECLARE_DELEGATE_OneParam(FOnControlBusesRemoved, const TArray<FName>& /*RemovedControlBusesNames*/);
+		FOnControlBusesRemoved OnControlBusesRemoved;
+
+	private:
+		void UpdateActiveControlBusesToAdd(const TMap<FBusId, FString>& InBusIdToBusNameMap);
+		void UpdateActiveControlBusesToRemove(const TMap<FBusId, float>& InBusIdToValueMap);
+
+#if WITH_EDITOR
+		void OnAudioDeviceDestroyed(::Audio::FDeviceId InDeviceId);
+#endif // WITH_EDITOR
+
+		FModulationMatrixMessages TraceMessages;
+
+		TMap<::Audio::FDeviceId, TSet<uint32>> DeviceIdToActiveModulatorSourceIdsMap;
+		BusIdToBusInfoMap ActiveControlBuses;
+		TArray<FName> RemovedControlBusesNames;
+	};
+} // namespace AudioModulationInsights
diff --git a/Engine/Plugins/Runtime/AudioModulationInsights/Source/AudioModulationInsights/Private/Views/AudioModulationInsightsCommands.cpp b/Engine/Plugins/Runtime/AudioModulationInsights/Source/AudioModulationInsights/Private/Views/AudioModulationInsightsCommands.cpp
new file mode 100644
--- /dev/null
+++ b/Engine/Plugins/Runtime/AudioModulationInsights/Source/AudioModulationInsights/Private/Views/AudioModulationInsightsCommands.cpp
@@ -0,0 +1,24 @@
+// Copyright Epic Games, Inc. All Rights Reserved.
+#include "AudioModulationInsightsCommands.h"
+
+#include "AudioInsightsStyle.h"
+#include "Framework/Commands/InputChord.h"
+#include "Framework/Commands/UICommandInfo.h"
+
+#define LOCTEXT_NAMESPACE "AudioModulationInsights"
+
+namespace AudioModulationInsights
+{
+	FAudioModulationInsightsCommands::FAudioModulationInsightsCommands()
+		: TCommands<FAudioModulationInsightsCommands>("AudioModulationInsightsCommands", LOCTEXT("AudioModulationInsightsCommands_ContextDescText", "Audio Modulation Insights Commands"), NAME_None, UE::Audio::Insights::FSlateStyle::GetStyleName())
+	{
+	}
+	
+	void FAudioModulationInsightsCommands::RegisterCommands()
+	{
+		UI_COMMAND(Browse, "Browse To Asset", "Browses to the selected sound asset in the content browser.", EUserInterfaceActionType::Button, FInputChord(EKeys::B, EModifierKey::Control));
+		UI_COMMAND(Edit, "Edit", "Opens the selected sound for edit.", EUserInterfaceActionType::Button, FInputChord(EKeys::E, EModifierKey::Control));
+	}
+} // namespace AudioModulationInsights
+
+#undef LOCTEXT_NAMESPACE
diff --git a/Engine/Plugins/Runtime/AudioModulationInsights/Source/AudioModulationInsights/Private/Views/AudioModulationInsightsCommands.h b/Engine/Plugins/Runtime/AudioModulationInsights/Source/AudioModulationInsights/Private/Views/AudioModulationInsightsCommands.h
new file mode 100644
--- /dev/null
+++ b/Engine/Plugins/Runtime/AudioModulationInsights/Source/AudioModulationInsights/Private/Views/AudioModulationInsightsCommands.h
@@ -0,0 +1,25 @@
+// Copyright Epic Games, Inc. All Rights Reserved.
+#pragma once
+
+#include "Framework/Commands/Commands.h"
+#include "Templates/SharedPointer.h"
+
+class FUICommandInfo;
+
+namespace AudioModulationInsights
+{
+	class FAudioModulationInsightsCommands : public TCommands<FAudioModulationInsightsCommands>
+	{
+	public:
+		FAudioModulationInsightsCommands();
+
+		virtual void RegisterCommands() override;
+
+		TSharedPtr<const FUICommandInfo> GetBrowseCommand() const { return Browse; }
+		TSharedPtr<const FUICommandInfo> GetEditCommand() const   { return Edit; }
+		
+	private:
+		TSharedPtr<FUICommandInfo> Browse;
+		TSharedPtr<FUICommandInfo> Edit;
+	};
+} // namespace AudioModulationInsights
diff --git a/Engine/Plugins/Runtime/AudioModulationInsights/Source/AudioModulationInsights/Private/Views/ControlBusDashboardViewFactory.cpp b/Engine/Plugins/Runtime/AudioModulationInsights/Source/AudioModulationInsights/Private/Views/ControlBusDashboardViewFactory.cpp
new file mode 100644
--- /dev/null
+++ b/Engine/Plugins/Runtime/AudioModulationInsights/Source/AudioModulationInsights/Private/Views/ControlBusDashboardViewFactory.cpp
@@ -0,0 +1,308 @@
+// Copyright Epic Games, Inc. All Rights Reserved.
+#include "ControlBusDashboardViewFactory.h"
+
+#include "AudioModulationInsightsCommands.h"
+#include "AudioModulationInsightsStyle.h"
+#include "AudioInsightsStyle.h"
+#include "AudioInsightsTraceModule.h"
+#include "AudioInsightsTraceProviderBase.h"
+#include "Framework/MultiBox/MultiBoxBuilder.h"
+#include "IAudioInsightsModule.h"
+#include "Internationalization/Text.h"
+#include "Providers/ControlBusTraceProvider.h"
+#include "Templates/SharedPointer.h"
+
+#define LOCTEXT_NAMESPACE "AudioModulationInsights"
+
+namespace AudioModulationInsights
+{
+	namespace ControlBusPrivate
+	{
+		const FControlBusDashboardEntry& CastEntry(const UE::Audio::Insights::IDashboardDataViewEntry& InData)
+		{
+			return static_cast<const FControlBusDashboardEntry&>(InData);
+		};
+
+	} // namespace ControlBusPrivate
+
+	FControlBusDashboardViewFactory::FControlBusDashboardViewFactory()
+	{
+		IAudioInsightsModule& InsightsModule = IAudioInsightsModule::GetChecked();
+		IAudioInsightsTraceModule& InsightsTraceModule = InsightsModule.GetTraceModule();
+
+		TSharedPtr<FControlBusTraceProvider> ControlBusProvider = MakeShared<FControlBusTraceProvider>();
+		InsightsTraceModule.AddTraceProvider(StaticCastSharedPtr<UE::Audio::Insights::FTraceProviderBase>(ControlBusProvider));
+
+		Providers =
+		{
+			StaticCastSharedPtr<UE::Audio::Insights::FTraceProviderBase>(ControlBusProvider)
+		};
+
+		FAudioModulationInsightsCommands::Register();
+
+		BindCommands();
+	}
+
+	FControlBusDashboardViewFactory::~FControlBusDashboardViewFactory()
+	{
+		FAudioModulationInsightsCommands::Unregister();
+	}
+
+	FName FControlBusDashboardViewFactory::GetName() const
+	{
+		return "ControlBuses";
+	}
+
+	FText FControlBusDashboardViewFactory::GetDisplayName() const
+	{
+		return LOCTEXT("AudioInsights_ModulationControlBus_DisplayName", "Control Buses");
+	}
+
+	void FControlBusDashboardViewFactory::ProcessEntries(UE::Audio::Insights::FTraceTableDashboardViewFactory::EProcessReason Reason)
+	{
+		using namespace UE::Audio::Insights;
+
+		const FString FilterString = GetSearchFilterText().ToString();
+		FTraceTableDashboardViewFactory::FilterEntries<FControlBusTraceProvider>([&FilterString](const IDashboardDataViewEntry& Entry)
+		{
+			const FControlBusDashboardEntry& ControlBusEntry = ControlBusPrivate::CastEntry(Entry);
+
+			return ControlBusEntry.DisplayName.Contains(FilterString);
+		});
+	}
+
+	FSlateIcon FControlBusDashboardViewFactory::GetIcon() const
+	{
+		return FAudioModulationInsightsStyle::Get().CreateIcon("ControlBus");
+	}
+
+	UE::Audio::Insights::EDefaultDashboardTabStack FControlBusDashboardViewFactory::GetDefaultTabStack() const
+	{
+		return UE::Audio::Insights::EDefaultDashboardTabStack::Analysis;
+	}
+
+	TSharedRef<SWidget> FControlBusDashboardViewFactory::MakeWidget(TSharedRef<SDockTab> OwnerTab, const FSpawnTabArgs& SpawnTabArgs)
+	{
+		TSharedRef<SWidget> TableDashboardWidget = UE::Audio::Insights::FTraceTableDashboardViewFactory::MakeWidget(OwnerTab, SpawnTabArgs);
+		TSharedRef<SWidget> BusWatchWidget = MakeControlBusWatchWidget();
+
+		return SNew(SHorizontalBox)
+			+ SHorizontalBox::Slot()
+			.HAlign(HAlign_Fill)
+			.Padding(0.0f, 6.0f, 0.0f, 0.0f)
+			[
+				// Dashboard and plots area
+				SNew(SSplitter)
+					.Orientation(Orient_Horizontal)
+					+ SSplitter::Slot()
+					.Value(0.55f)
+					[
+						TableDashboardWidget
+					]
+					+ SSplitter::Slot()
+					.Value(0.45f)
+					[
+						BusWatchWidget
+					]
+			];
+	}
+
+	const TMap<FName, UE::Audio::Insights::FTraceTableDashboardViewFactory::FColumnData>& FControlBusDashboardViewFactory::GetColumns() const
+	{
+		auto CreateColumnData = []()
+		{
+			using namespace UE::Audio::Insights;
+
+			return TMap<FName, FTraceTableDashboardViewFactory::FColumnData>
+			{
+				{
+					"BusId",
+					{
+						LOCTEXT("ControlBus_BusIdColumnDisplayName", "Control Bus ID"),
+						[](const IDashboardDataViewEntry& InData) { return FText::AsNumber(ControlBusPrivate::CastEntry(InData).ControlBusId); },
+						nullptr	/* GetIconName */,
+						true /* bDefaultHidden */,
+						0.08f /* FillWidth */
+					}
+				},
+				{
+					"Name",
+					{
+						LOCTEXT("ControlBus_NameColumnDisplayName", "Name"),
+						[](const IDashboardDataViewEntry& InData) { return ControlBusPrivate::CastEntry(InData).GetDisplayNameAsFText(); },
+						nullptr	/* GetIconName */,
+						false /* bDefaultHidden */,
+						0.75f /* FillWidth */
+					}
+				},
+				{
+					"ParameterName",
+					{
+						LOCTEXT("ControlBus_ParamNameColumnDisplayName", "Parameter"),
+						[](const IDashboardDataViewEntry& InData) { return FText::FromString(ControlBusPrivate::CastEntry(InData).ParamName); },
+						nullptr	/* GetIconName */,
+						false /* bDefaultHidden */,
+						0.15f /* FillWidth */
+					}
+				},
+				{
+					"Value",
+					{
+						LOCTEXT("ControlBus_ValueColumnDisplayName", "Value"),
+						[](const IDashboardDataViewEntry& InData) { return FText::AsNumber(ControlBusPrivate::CastEntry(InData).Value, FSlateStyle::Get().GetLinearVolumeFloatFormat()); },
+						nullptr	/* GetIconName */,
+						false /* bDefaultHidden */,
+						0.07f /* FillWidth */
+					}
+				}
+			};
+		};
+		static const TMap<FName, FTraceTableDashboardViewFactory::FColumnData> ColumnData = CreateColumnData();
+		return ColumnData;
+	}
+
+	void FControlBusDashboardViewFactory::SortTable()
+	{
+		using namespace UE::Audio::Insights;
+
+		if (SortByColumn == "BusId")
+		{
+			if (SortMode == EColumnSortMode::Ascending)
+			{
+				DataViewEntries.Sort([](const TSharedPtr<IDashboardDataViewEntry>& A, const TSharedPtr<IDashboardDataViewEntry>& B)
+				{
+					const FControlBusDashboardEntry& AData = ControlBusPrivate::CastEntry(*A.Get());
+					const FControlBusDashboardEntry& BData = ControlBusPrivate::CastEntry(*B.Get());
+
+					return AData.ControlBusId < BData.ControlBusId;
+				});
+			}
+			else if (SortMode == EColumnSortMode::Descending)
+			{
+				DataViewEntries.Sort([](const TSharedPtr<IDashboardDataViewEntry>& A, const TSharedPtr<IDashboardDataViewEntry>& B)
+				{
+					const FControlBusDashboardEntry& AData = ControlBusPrivate::CastEntry(*A.Get());
+					const FControlBusDashboardEntry& BData = ControlBusPrivate::CastEntry(*B.Get());
+
+					return BData.ControlBusId < AData.ControlBusId;
+				});
+			}
+		}
+		else if (SortByColumn == "Name")
+		{
+			if (SortMode == EColumnSortMode::Ascending)
+			{
+				DataViewEntries.Sort([](const TSharedPtr<IDashboardDataViewEntry>& A, const TSharedPtr<IDashboardDataViewEntry>& B)
+				{
+					const FControlBusDashboardEntry& AData = ControlBusPrivate::CastEntry(*A.Get());
+					const FControlBusDashboardEntry& BData = ControlBusPrivate::CastEntry(*B.Get());
+
+					return AData.DisplayName.ToLower() < BData.DisplayName.ToLower();
+				});
+			}
+			else if (SortMode == EColumnSortMode::Descending)
+			{
+				DataViewEntries.Sort([](const TSharedPtr<IDashboardDataViewEntry>& A, const TSharedPtr<IDashboardDataViewEntry>& B)
+				{
+					const FControlBusDashboardEntry& AData = ControlBusPrivate::CastEntry(*A.Get());
+					const FControlBusDashboardEntry& BData = ControlBusPrivate::CastEntry(*B.Get());
+
+					return BData.DisplayName.ToLower() < AData.DisplayName.ToLower();
+				});
+			}
+		}
+		else if (SortByColumn == "ParamName")
+		{
+			if (SortMode == EColumnSortMode::Ascending)
+			{
+				DataViewEntries.Sort([](const TSharedPtr<IDashboardDataViewEntry>& A, const TSharedPtr<IDashboardDataViewEntry>& B)
+				{
+					const FControlBusDashboardEntry& AData = ControlBusPrivate::CastEntry(*A.Get());
+					const FControlBusDashboardEntry& BData = ControlBusPrivate::CastEntry(*B.Get());
+
+					return AData.ParamName.ToLower() < BData.ParamName.ToLower();
+				});
+			}
+			else if (SortMode == EColumnSortMode::Descending)
+			{
+				DataViewEntries.Sort([](const TSharedPtr<IDashboardDataViewEntry>& A, const TSharedPtr<IDashboardDataViewEntry>& B)
+				{
+					const FControlBusDashboardEntry& AData = ControlBusPrivate::CastEntry(*A.Get());
+					const FControlBusDashboardEntry& BData = ControlBusPrivate::CastEntry(*B.Get());
+
+					return BData.ParamName.ToLower() < AData.ParamName.ToLower();
+				});
+			}
+		}
+		else if (SortByColumn == "Value")
+		{
+			if (SortMode == EColumnSortMode::Ascending)
+			{
+				DataViewEntries.Sort([](const TSharedPtr<IDashboardDataViewEntry>& A, const TSharedPtr<IDashboardDataViewEntry>& B)
+				{
+					const FControlBusDashboardEntry& AData = ControlBusPrivate::CastEntry(*A.Get());
+					const FControlBusDashboardEntry& BData = ControlBusPrivate::CastEntry(*B.Get());
+
+					return AData.Value < BData.Value;
+				});
+			}
+			else if (SortMode == EColumnSortMode::Descending)
+			{
+				DataViewEntries.Sort([](const TSharedPtr<IDashboardDataViewEntry>& A, const TSharedPtr<IDashboardDataViewEntry>& B)
+				{
+					const FControlBusDashboardEntry& AData = ControlBusPrivate::CastEntry(*A.Get());
+					const FControlBusDashboardEntry& BData = ControlBusPrivate::CastEntry(*B.Get());
+
+					return BData.Value < AData.Value;
+				});
+			}
+		}
+	}
+
+	TSharedRef<SWidget> FControlBusDashboardViewFactory::MakeControlBusListWidget()
+	{
+		return SNew(SBox);
+	}
+
+	TSharedRef<SWidget> FControlBusDashboardViewFactory::MakeControlBusWatchWidget()
+	{
+		return SNew(SBox);
+	}
+
+	TSharedPtr<SWidget> FControlBusDashboardViewFactory::OnConstructContextMenu()
+	{
+		const FAudioModulationInsightsCommands& Commands = FAudioModulationInsightsCommands::Get();
+
+		constexpr bool bShouldCloseWindowAfterMenuSelection = true;
+		FMenuBuilder MenuBuilder(bShouldCloseWindowAfterMenuSelection, CommandList);
+
+		MenuBuilder.BeginSection("ControlBusDashboardActions", LOCTEXT("ControlBusActions_HeaderText", "Control Bus Options"));
+		{
+			MenuBuilder.AddMenuEntry(Commands.GetBrowseCommand(), NAME_None, TAttribute<FText>(), TAttribute<FText>(), UE::Audio::Insights::FSlateStyle::Get().CreateIcon("AudioInsights.Icon.SoundDashboard.Browse"));
+			MenuBuilder.AddMenuEntry(Commands.GetEditCommand(), NAME_None, TAttribute<FText>(), TAttribute<FText>(), UE::Audio::Insights::FSlateStyle::Get().CreateIcon("AudioInsights.Icon.SoundDashboard.Edit"));
+		}
+
+		MenuBuilder.EndSection();
+
+		return MenuBuilder.MakeWidget();
+	}
+
+	FReply FControlBusDashboardViewFactory::OnDataRowKeyInput(const FGeometry& InGeometry, const FKeyEvent& InKeyEvent) const
+	{
+		return (CommandList && CommandList->ProcessCommandBindings(InKeyEvent)) ? FReply::Handled() : FReply::Unhandled();
+	}
+
+	void FControlBusDashboardViewFactory::BindCommands()
+	{
+		CommandList = MakeShared<FUICommandList>();
+
+		const FAudioModulationInsightsCommands& Commands = FAudioModulationInsightsCommands::Get();
+
+#if WITH_EDITOR
+		CommandList->MapAction(Commands.GetBrowseCommand(), FExecuteAction::CreateLambda([this]() { BrowseToAsset(); }));
+		CommandList->MapAction(Commands.GetEditCommand(), FExecuteAction::CreateLambda([this]() { OpenAsset(); }));
+#endif // WITH_EDITOR
+	}
+} // namespace AudioModulationInsights
+
+#undef LOCTEXT_NAMESPACE
diff --git a/Engine/Plugins/Runtime/AudioModulationInsights/Source/AudioModulationInsights/Private/Views/ControlBusDashboardViewFactory.h b/Engine/Plugins/Runtime/AudioModulationInsights/Source/AudioModulationInsights/Private/Views/ControlBusDashboardViewFactory.h
new file mode 100644
--- /dev/null
+++ b/Engine/Plugins/Runtime/AudioModulationInsights/Source/AudioModulationInsights/Private/Views/ControlBusDashboardViewFactory.h
@@ -0,0 +1,38 @@
+// Copyright Epic Games, Inc. All Rights Reserved.
+#pragma once
+
+#include "Templates/SharedPointer.h"
+#include "Views/TableDashboardViewFactory.h"
+
+namespace AudioModulationInsights
+{
+	class FControlBusDashboardViewFactory : public UE::Audio::Insights::FTraceObjectTableDashboardViewFactory
+	{
+	public:
+		FControlBusDashboardViewFactory();
+		virtual ~FControlBusDashboardViewFactory();
+
+		virtual FName GetName() const override;
+		virtual FText GetDisplayName() const override;
+		virtual FSlateIcon GetIcon() const override;
+		virtual UE::Audio::Insights::EDefaultDashboardTabStack GetDefaultTabStack() const override;
+		virtual void ProcessEntries(UE::Audio::Insights::FTraceTableDashboardViewFactory::EProcessReason Reason) override;
+		virtual TSharedRef<SWidget> MakeWidget(TSharedRef<SDockTab> OwnerTab, const FSpawnTabArgs& SpawnTabArgs) override;
+
+	protected:
+		virtual const TMap<FName, UE::Audio::Insights::FTraceTableDashboardViewFactory::FColumnData>& GetColumns() const override;
+
+		virtual void SortTable() override;
+
+		virtual TSharedRef<SWidget> MakeControlBusListWidget();
+		virtual TSharedRef<SWidget> MakeControlBusWatchWidget();
+
+		virtual TSharedPtr<SWidget> OnConstructContextMenu() override;
+		virtual FReply OnDataRowKeyInput(const FGeometry& InGeometry, const FKeyEvent& InKeyEvent) const override;
+
+	private:
+		void BindCommands();
+
+		TSharedPtr<FUICommandList> CommandList;
+	};
+} // namespace AudioModulationInsights
diff --git a/Engine/Plugins/Runtime/AudioModulationInsights/Source/AudioModulationInsights/Private/Views/ModulationMatrixDashboardViewFactory.cpp b/Engine/Plugins/Runtime/AudioModulationInsights/Source/AudioModulationInsights/Private/Views/ModulationMatrixDashboardViewFactory.cpp
new file mode 100644
--- /dev/null
+++ b/Engine/Plugins/Runtime/AudioModulationInsights/Source/AudioModulationInsights/Private/Views/ModulationMatrixDashboardViewFactory.cpp
@@ -0,0 +1,543 @@
+// Copyright Epic Games, Inc. All Rights Reserved.
+#include "ModulationMatrixDashboardViewFactory.h"
+
+#include "AudioInsightsStyle.h"
+#include "AudioInsightsTraceProviderBase.h"
+#include "AudioModulationInsightsStyle.h"
+#include "AudioModulationInsightsCommands.h"
+#include "Framework/Application/SlateApplication.h"
+#include "Framework/MultiBox/MultiBoxBuilder.h"
+#include "Fonts/FontMeasure.h"
+#include "IAudioInsightsModule.h"
+#include "Internationalization/Text.h"
+#include "Rendering/SlateRenderer.h"
+#include "Templates/SharedPointer.h"
+#include "Widgets/Input/SComboBox.h"
+
+#if WITH_EDITOR
+#include "AudioDeviceManager.h"
+#endif // WITH_EDITOR
+
+#define LOCTEXT_NAMESPACE "AudioModulationInsights"
+
+namespace AudioModulationInsights
+{
+	namespace FModulationMatrixDashboardViewFactoryPrivate
+	{
+		const FName SourceIdColumnName("SourceId");
+		const FName ModulatingSourceColumnName("ModulatingSource");
+		const FName ModulatingSourceTypeColumnName("EntryType");
+
+		const FModulationMatrixDashboardEntry& CastEntry(const UE::Audio::Insights::IDashboardDataViewEntry& InData)
+		{
+			return static_cast<const FModulationMatrixDashboardEntry&>(InData);
+		}		
+	}
+
+	FModulationMatrixDashboardViewFactory::FModulationMatrixDashboardViewFactory()
+	{
+		IAudioInsightsTraceModule& AudioInsightsTraceModule = IAudioInsightsModule::GetChecked().GetTraceModule();
+
+		ModulationMatrixTraceProvider = MakeShared<FModulationMatrixTraceProvider>();
+
+		AudioInsightsTraceModule.AddTraceProvider(ModulationMatrixTraceProvider);
+
+		Providers =
+		{
+			ModulationMatrixTraceProvider
+		};
+
+		CreateDefaultColumnData();
+
+		FAudioModulationInsightsCommands::Register();
+
+		BindCommands();
+	}
+
+	FModulationMatrixDashboardViewFactory::~FModulationMatrixDashboardViewFactory()
+	{
+		FAudioModulationInsightsCommands::Unregister();
+	}
+
+	FName FModulationMatrixDashboardViewFactory::GetName() const
+	{
+		static const FName ModulationMatrixDashboardName = "ModulationMatrix";
+		return ModulationMatrixDashboardName;
+	}
+
+	FText FModulationMatrixDashboardViewFactory::GetDisplayName() const
+	{
+		static const FText ModulationMatrixDashboardDisplayName = LOCTEXT("AudioInsights_ModulationMatrix_DisplayName", "Modulation Matrix");
+		return ModulationMatrixDashboardDisplayName;
+	}
+
+	FSlateIcon FModulationMatrixDashboardViewFactory::GetIcon() const
+	{
+		return FAudioModulationInsightsStyle::Get().CreateIcon("ControlBusMix");
+	}
+
+	UE::Audio::Insights::EDefaultDashboardTabStack FModulationMatrixDashboardViewFactory::GetDefaultTabStack() const
+	{
+		return UE::Audio::Insights::EDefaultDashboardTabStack::Analysis;
+	}
+
+	const TMap<FName, UE::Audio::Insights::FTraceTableDashboardViewFactory::FColumnData>& FModulationMatrixDashboardViewFactory::GetColumns() const
+	{
+		return ModulationMatrixColumnData;
+	}
+
+	void FModulationMatrixDashboardViewFactory::CreateDefaultColumnData()
+	{
+		using namespace UE::Audio::Insights;
+		using namespace FModulationMatrixDashboardViewFactoryPrivate;
+
+		if (!ModulationMatrixColumnData.IsEmpty())
+		{
+			return;
+		}
+
+		ModulationMatrixColumnData =
+		{
+			{
+				SourceIdColumnName,
+				{
+					LOCTEXT("ModulationMatrix_SourceIdColumnDisplayName", "Source Id"),
+					[](const IDashboardDataViewEntry& InData)
+					{
+						const FSourceId SourceId = CastEntry(InData).SourceId;
+						return SourceId != INDEX_NONE ? FText::AsNumber(SourceId) : FText::GetEmpty();
+					},
+					nullptr	/* GetIconName */,
+					true /* bDefaultHidden */,
+					0.1f /* FillWidth */
+				}
+			},
+			{
+				ModulatingSourceColumnName,
+				{
+					LOCTEXT("ModulationMatrix_ModulatingSourceDisplayName", "Modulating Source"),
+					[](const IDashboardDataViewEntry& InData) { return FText::FromString(CastEntry(InData).DisplayName); },
+					nullptr	/* GetIconName */,
+					false /* bDefaultHidden */,
+					0.4f /* FillWidth */
+				}
+			},
+			{
+				ModulatingSourceTypeColumnName,
+				{
+					LOCTEXT("ModulationMatrix_ModulatingSourceTypeDisplayName", "Type"),
+					[](const IDashboardDataViewEntry& InData)
+					{
+						const EModulationMatrixEntryType EntryType = CastEntry(InData).EntryType;
+
+						switch (EntryType)
+						{
+							case EModulationMatrixEntryType::BusMix:
+								return LOCTEXT("ModulationMatrix_EntryTypeBusMix", "Bus Mix");
+
+							case EModulationMatrixEntryType::Generator:
+								return LOCTEXT("ModulationMatrix_EntryTypeGenerator", "Generator");
+
+							default:
+								break;
+						}
+
+						return FText::GetEmpty();
+					},
+					nullptr	/* GetIconName */,
+					false /* bDefaultHidden */,
+					0.125f /* FillWidth */
+				}
+			}
+		};
+
+		SortByColumn = ModulatingSourceColumnName;
+		SortMode = EColumnSortMode::Ascending;
+	}
+
+	void FModulationMatrixDashboardViewFactory::RegisterDelegates()
+	{
+		if (!ModulationMatrixTraceProvider->OnControlBusesAdded.IsBound())
+		{
+			ModulationMatrixTraceProvider->OnControlBusesAdded.BindSP(this, &FModulationMatrixDashboardViewFactory::OnControlBusesAdded);
+		}
+
+		if (!ModulationMatrixTraceProvider->OnControlBusesRemoved.IsBound())
+		{
+			ModulationMatrixTraceProvider->OnControlBusesRemoved.BindSP(this, &FModulationMatrixDashboardViewFactory::OnControlBusesRemoved);
+		}
+
+#if WITH_EDITOR
+		if (!OnDeviceDestroyedHandle.IsValid())
+		{
+			OnDeviceDestroyedHandle = FAudioDeviceManagerDelegates::OnAudioDeviceDestroyed.AddSP(this, &FModulationMatrixDashboardViewFactory::OnAudioDeviceDestroyed);
+		}
+#endif // WITH_EDITOR
+	}
+
+	void FModulationMatrixDashboardViewFactory::BindCommands()
+	{
+		CommandList = MakeShared<FUICommandList>();
+
+		const FAudioModulationInsightsCommands& Commands = FAudioModulationInsightsCommands::Get();
+
+#if WITH_EDITOR
+		CommandList->MapAction(Commands.GetBrowseCommand(), FExecuteAction::CreateLambda([this]() { BrowseToAsset(); }));
+		CommandList->MapAction(Commands.GetEditCommand(), FExecuteAction::CreateLambda([this]() { OpenAsset(); }));
+#endif // WITH_EDITOR
+	}
+
+	TSharedRef<SWidget> FModulationMatrixDashboardViewFactory::MakeModulatingSourceTypeFilterWidget()
+	{
+		if (ModulatingSourceTypes.IsEmpty())
+		{
+			ModulatingSourceTypes.Emplace(MakeShared<FComboboxSelectionItem>(EModulatingSourceComboboxSelection::All,        LOCTEXT("ModulationMatrix_ModulatingSourceTypeAll",        "All")));
+			ModulatingSourceTypes.Emplace(MakeShared<FComboboxSelectionItem>(EModulatingSourceComboboxSelection::BusMixes,   LOCTEXT("ModulationMatrix_ModulatingSourceTypeBusMixes",   "Bus Mixes")));
+			ModulatingSourceTypes.Emplace(MakeShared<FComboboxSelectionItem>(EModulatingSourceComboboxSelection::Generators, LOCTEXT("ModulationMatrix_ModulatingSourceTypeGenerators", "Generators")));
+
+			SelectedModulatingSourceType = ModulatingSourceTypes[0];
+		}
+
+		return SNew(SHorizontalBox)
+			+ SHorizontalBox::Slot()
+			.AutoWidth()
+			.Padding(2.0f, 10.0f, 0.0f, 0.0f)
+			[
+				SNew(STextBlock)
+				.Margin(FMargin(0.0, 2.0, 0.0, 0.0))
+				.Text(LOCTEXT("ModulationMatrix_TypeFilterText", "Type Filter:"))
+			]
+			+ SHorizontalBox::Slot()
+			.MaxWidth(2.0f)
+			.Padding(0.0f, 10.0f, 0.0f, 0.0f)
+			[
+				SNew(SBox)
+			]
+			+ SHorizontalBox::Slot()
+			.AutoWidth()
+			.HAlign(EHorizontalAlignment::HAlign_Center)
+			.Padding(0.0f, 10.0f, 0.0f, 0.0f)
+			[
+				SNew(SComboBox<TSharedPtr<FComboboxSelectionItem>>)
+				.OptionsSource(&ModulatingSourceTypes)
+				.OnGenerateWidget_Lambda([this](const TSharedPtr<FComboboxSelectionItem>& ModulatingSourceTypePtr)
+				{
+					const FText ModulatingSourceTypeDisplayName = ModulatingSourceTypePtr.IsValid() ? ModulatingSourceTypePtr->Value /*DisplayName*/ : FText::GetEmpty();
+
+					return SNew(STextBlock)
+						.Text(ModulatingSourceTypeDisplayName);
+				})
+				.OnSelectionChanged_Lambda([this](TSharedPtr<FComboboxSelectionItem> InSelectedModulatingSourceTypePtr, ESelectInfo::Type)
+				{
+					if (InSelectedModulatingSourceTypePtr.IsValid())
+					{
+						SelectedModulatingSourceType = InSelectedModulatingSourceTypePtr;
+						UpdateFilterReason = EProcessReason::FilterUpdated;
+					}
+				})
+				[
+					SNew(STextBlock)
+					.Text_Lambda([this]()
+					{
+						const int32 FoundIndex = ModulatingSourceTypes.Find(SelectedModulatingSourceType);
+						if (ModulatingSourceTypes.IsValidIndex(FoundIndex) && ModulatingSourceTypes[FoundIndex].IsValid())
+						{
+							return ModulatingSourceTypes[FoundIndex]->Value;
+						}
+
+						return FText::GetEmpty();
+					})
+				]
+			];
+	}
+
+	TSharedRef<SWidget> FModulationMatrixDashboardViewFactory::MakeWidget(TSharedRef<SDockTab> OwnerTab, const FSpawnTabArgs& SpawnTabArgs)
+	{
+		using namespace UE::Audio::Insights;
+
+		RegisterDelegates();
+
+		return SNew(SVerticalBox)
+		+ SVerticalBox::Slot()
+		.AutoHeight()
+		.HAlign(HAlign_Fill)
+		.Padding(0.0f, 0.0f, 0.0f, 6.0f)
+		[
+			MakeModulatingSourceTypeFilterWidget()
+		]
+		+ SVerticalBox::Slot()
+		.HAlign(HAlign_Fill)
+		[
+			FTraceTableDashboardViewFactory::MakeWidget(OwnerTab, SpawnTabArgs)
+		];
+	}
+
+	void FModulationMatrixDashboardViewFactory::ProcessEntries(UE::Audio::Insights::FTraceTableDashboardViewFactory::EProcessReason Reason)
+	{
+		FilterByModulatingSourceName();
+		FilterByModulatingSourceType();
+	}
+
+	TSharedPtr<SWidget> FModulationMatrixDashboardViewFactory::OnConstructContextMenu()
+	{
+		const FAudioModulationInsightsCommands& Commands = FAudioModulationInsightsCommands::Get();
+
+		constexpr bool bShouldCloseWindowAfterMenuSelection = true;
+		FMenuBuilder MenuBuilder(bShouldCloseWindowAfterMenuSelection, CommandList);
+
+		MenuBuilder.BeginSection("ControlBusDashboardActions", LOCTEXT("AudioInsights_ModulationMatrix_HeaderText", "Modulation Matrix Options"));
+		{
+			MenuBuilder.AddMenuEntry(Commands.GetBrowseCommand(), NAME_None, TAttribute<FText>(), TAttribute<FText>(), UE::Audio::Insights::FSlateStyle::Get().CreateIcon("AudioInsights.Icon.SoundDashboard.Browse"));
+			MenuBuilder.AddMenuEntry(Commands.GetEditCommand(), NAME_None, TAttribute<FText>(), TAttribute<FText>(), UE::Audio::Insights::FSlateStyle::Get().CreateIcon("AudioInsights.Icon.SoundDashboard.Edit"));
+		}
+
+		MenuBuilder.EndSection();
+
+		return MenuBuilder.MakeWidget();
+	}
+
+	FReply FModulationMatrixDashboardViewFactory::OnDataRowKeyInput(const FGeometry& InGeometry, const FKeyEvent& InKeyEvent) const
+	{
+		return (CommandList && CommandList->ProcessCommandBindings(InKeyEvent)) ? FReply::Handled() : FReply::Unhandled();
+	}
+
+	void FModulationMatrixDashboardViewFactory::OnControlBusesAdded(const FModulationMatrixTraceProvider::BusIdToBusInfoMap& AddedControlBuses)
+	{
+		using namespace UE::Audio::Insights;
+		using namespace FModulationMatrixDashboardViewFactoryPrivate;
+
+		for (const auto& [BusId, BusInfo] : AddedControlBuses)
+		{
+			const FName BusFName(BusInfo.BusName);
+
+			if (HeaderRowWidget->IsColumnGenerated(BusFName))
+			{
+				continue;
+			}
+
+			const FText BusNameText = FText::FromString(BusInfo.BusName);
+
+			// Add bus infomation to column data
+			ModulationMatrixColumnData.Add(BusFName,
+			{
+				BusNameText,
+				[BusId](const IDashboardDataViewEntry& InData)
+				{
+					const FModulationMatrixDashboardEntry& ModulationMatrixEntry = CastEntry(InData);
+
+					if (const float* FoundValuePtr = ModulationMatrixEntry.BusIdToValueMap.Find(BusId))
+					{
+						return FText::AsNumber(*FoundValuePtr, FSlateStyle::Get().GetLinearVolumeFloatFormat());
+					}
+
+					return FText::GetEmpty();
+				},
+				nullptr	/* GetIconName */,
+				false /* bDefaultHidden */,
+				0.1f  /* FillWidth */
+			});
+
+			// Add column to widget
+			SHeaderRow::FColumn::FArguments ColumnArgs = SHeaderRow::Column(BusFName)
+				.DefaultLabel(BusNameText)
+				.HAlignCell(HAlign_Left);
+
+			if (const FSlateRenderer* SlateRenderer = FSlateApplication::Get().GetRenderer())
+			{
+				const FVector2D TextSize = SlateRenderer->GetFontMeasureService()->Measure(BusNameText, FAppStyle::GetFontStyle("NormalFont"));
+
+				ColumnArgs.ManualWidth(TextSize.X + 10.0);
+			}
+
+			HeaderRowWidget->AddColumn(ColumnArgs);
+
+			ActiveBusNames.Emplace(BusFName);
+		}
+	}
+
+	void FModulationMatrixDashboardViewFactory::OnControlBusesRemoved(const TArray<FName>& RemovedControlBusesNames)
+	{
+		for (const FName& BusName : RemovedControlBusesNames)
+		{
+			HeaderRowWidget->RemoveColumn(BusName);
+			ModulationMatrixColumnData.Remove(BusName);
+			ActiveBusNames.Remove(BusName);
+		}
+	}
+
+#if WITH_EDITOR
+	void FModulationMatrixDashboardViewFactory::OnAudioDeviceDestroyed(::Audio::FDeviceId InDeviceId)
+	{
+		for (const FName& BusName : ActiveBusNames)
+		{
+			HeaderRowWidget->RemoveColumn(BusName);
+			ModulationMatrixColumnData.Remove(BusName);
+		}
+
+		ActiveBusNames.Empty();
+	}
+#endif // WITH_EDITOR
+
+	void FModulationMatrixDashboardViewFactory::FilterByModulatingSourceName()
+	{
+		using namespace UE::Audio::Insights;
+		
+		const FString FilterString = GetSearchFilterText().ToString();
+
+		FTraceTableDashboardViewFactory::FilterEntries<FModulationMatrixTraceProvider>(
+		[&FilterString](const IDashboardDataViewEntry& Entry)
+		{
+			const FModulationMatrixDashboardEntry& ModulationMatrixEntry = FModulationMatrixDashboardViewFactoryPrivate::CastEntry(Entry);
+
+			return ModulationMatrixEntry.DisplayName.Contains(FilterString) || ModulationMatrixEntry.EntryType == EModulationMatrixEntryType::BusFinalValues;
+		});
+	}
+
+	void FModulationMatrixDashboardViewFactory::FilterByModulatingSourceType()
+	{
+		using namespace UE::Audio::Insights;
+
+		TArray<TSharedPtr<IDashboardDataViewEntry>> EntriesToFilterOut;
+
+		const EModulatingSourceComboboxSelection SelectedModulatingSourceTypeEnum = SelectedModulatingSourceType.IsValid()? SelectedModulatingSourceType->Key : EModulatingSourceComboboxSelection::All;
+
+		for (const TSharedPtr<IDashboardDataViewEntry>& Entry : DataViewEntries)
+		{
+			if (Entry.IsValid())
+			{
+				const FModulationMatrixDashboardEntry& ModulationMatrixEntry = FModulationMatrixDashboardViewFactoryPrivate::CastEntry(*Entry);
+
+				// We never want to filter out the bus final values row
+				if (ModulationMatrixEntry.EntryType == EModulationMatrixEntryType::BusFinalValues)
+				{
+					continue;
+				}
+
+				if (SelectedModulatingSourceTypeEnum != EModulatingSourceComboboxSelection::All)
+				{
+					if ((SelectedModulatingSourceTypeEnum == EModulatingSourceComboboxSelection::BusMixes   && ModulationMatrixEntry.EntryType != EModulationMatrixEntryType::BusMix) ||
+						(SelectedModulatingSourceTypeEnum == EModulatingSourceComboboxSelection::Generators && ModulationMatrixEntry.EntryType != EModulationMatrixEntryType::Generator))
+					{
+						EntriesToFilterOut.Emplace(Entry);
+					}
+				}
+			}
+		}
+
+		for (const TSharedPtr<IDashboardDataViewEntry>& Entry : EntriesToFilterOut)
+		{
+			DataViewEntries.Remove(Entry);
+		}
+	}
+
+	void FModulationMatrixDashboardViewFactory::SortTable()
+	{
+		using namespace UE::Audio::Insights;
+		using namespace FModulationMatrixDashboardViewFactoryPrivate;
+
+		if (SortByColumn == SourceIdColumnName)
+		{
+			if (SortMode == EColumnSortMode::Ascending)
+			{
+				DataViewEntries.Sort([](const TSharedPtr<IDashboardDataViewEntry>& A, const TSharedPtr<IDashboardDataViewEntry>& B)
+				{
+					const FModulationMatrixDashboardEntry& AData = CastEntry(*A);
+					const FModulationMatrixDashboardEntry& BData = CastEntry(*B);
+
+					if (AData.EntryType == EModulationMatrixEntryType::BusFinalValues) { return false; }
+					if (BData.EntryType == EModulationMatrixEntryType::BusFinalValues) { return true;  }
+
+					return AData.SourceId < BData.SourceId;
+				});
+			}
+			else if (SortMode == EColumnSortMode::Descending)
+			{
+				DataViewEntries.Sort([](const TSharedPtr<IDashboardDataViewEntry>& A, const TSharedPtr<IDashboardDataViewEntry>& B)
+				{
+					const FModulationMatrixDashboardEntry& AData = CastEntry(*A);
+					const FModulationMatrixDashboardEntry& BData = CastEntry(*B);
+
+					if (AData.EntryType == EModulationMatrixEntryType::BusFinalValues) { return false; }
+					if (BData.EntryType == EModulationMatrixEntryType::BusFinalValues) { return true;  }
+
+					return BData.SourceId < AData.SourceId;
+				});
+			}
+		}
+		else if (SortByColumn == ModulatingSourceColumnName)
+		{
+			if (SortMode == EColumnSortMode::Ascending)
+			{
+				DataViewEntries.Sort([](const TSharedPtr<IDashboardDataViewEntry>& A, const TSharedPtr<IDashboardDataViewEntry>& B)
+				{
+					const FModulationMatrixDashboardEntry& AData = CastEntry(*A);
+					const FModulationMatrixDashboardEntry& BData = CastEntry(*B);
+
+					if (AData.EntryType == EModulationMatrixEntryType::BusFinalValues) { return false; }
+					if (BData.EntryType == EModulationMatrixEntryType::BusFinalValues) { return true;  }
+
+					return AData.Name.ToLower() < BData.Name.ToLower();
+				});
+			}
+			else if (SortMode == EColumnSortMode::Descending)
+			{
+				DataViewEntries.Sort([](const TSharedPtr<IDashboardDataViewEntry>& A, const TSharedPtr<IDashboardDataViewEntry>& B)
+				{
+					const FModulationMatrixDashboardEntry& AData = CastEntry(*A);
+					const FModulationMatrixDashboardEntry& BData = CastEntry(*B);
+
+					if (AData.EntryType == EModulationMatrixEntryType::BusFinalValues) { return false; }
+					if (BData.EntryType == EModulationMatrixEntryType::BusFinalValues) { return true;  }
+
+					return BData.Name.ToLower() < AData.Name.ToLower();
+				});
+			}
+		}
+		else if (SortByColumn == ModulatingSourceTypeColumnName)
+		{
+			if (SortMode == EColumnSortMode::Ascending)
+			{
+				DataViewEntries.Sort([](const TSharedPtr<IDashboardDataViewEntry>& A, const TSharedPtr<IDashboardDataViewEntry>& B)
+				{
+					const FModulationMatrixDashboardEntry& AData = CastEntry(*A);
+					const FModulationMatrixDashboardEntry& BData = CastEntry(*B);
+
+					if (AData.EntryType == EModulationMatrixEntryType::BusFinalValues) { return false; }
+					if (BData.EntryType == EModulationMatrixEntryType::BusFinalValues) { return true;  }
+
+					return AData.EntryType < BData.EntryType;
+				});
+			}
+			else if (SortMode == EColumnSortMode::Descending)
+			{
+				DataViewEntries.Sort([](const TSharedPtr<IDashboardDataViewEntry>& A, const TSharedPtr<IDashboardDataViewEntry>& B)
+				{
+					const FModulationMatrixDashboardEntry& AData = CastEntry(*A);
+					const FModulationMatrixDashboardEntry& BData = CastEntry(*B);
+
+					if (AData.EntryType == EModulationMatrixEntryType::BusFinalValues) { return false; }
+					if (BData.EntryType == EModulationMatrixEntryType::BusFinalValues) { return true;  }
+
+					return BData.EntryType < AData.EntryType;
+				});
+			}
+		}
+	}
+
+	FSlateColor FModulationMatrixDashboardViewFactory::GetRowColor(const TSharedPtr<UE::Audio::Insights::IDashboardDataViewEntry>& InRowDataPtr)
+	{
+		if (InRowDataPtr.IsValid())
+		{
+			const FModulationMatrixDashboardEntry& ModulationMatrixEntry = FModulationMatrixDashboardViewFactoryPrivate::CastEntry(*InRowDataPtr);
+
+			if (ModulationMatrixEntry.EntryType == EModulationMatrixEntryType::BusFinalValues)
+			{
+				return FSlateColor(FColor::Green);
+			}
+		}
+
+		return FSlateColor(FColor::White);
+	}
+} // namespace AudioModulationInsights
+
+#undef LOCTEXT_NAMESPACE
diff --git a/Engine/Plugins/Runtime/AudioModulationInsights/Source/AudioModulationInsights/Private/Views/ModulationMatrixDashboardViewFactory.h b/Engine/Plugins/Runtime/AudioModulationInsights/Source/AudioModulationInsights/Private/Views/ModulationMatrixDashboardViewFactory.h
new file mode 100644
--- /dev/null
+++ b/Engine/Plugins/Runtime/AudioModulationInsights/Source/AudioModulationInsights/Private/Views/ModulationMatrixDashboardViewFactory.h
@@ -0,0 +1,81 @@
+// Copyright Epic Games, Inc. All Rights Reserved.
+#pragma once
+
+#include "Containers/Map.h"
+#include "Providers/ModulationMatrixTraceProvider.h"
+#include "Templates/SharedPointer.h"
+#include "Views/TableDashboardViewFactory.h"
+
+namespace UE::Audio::Insights
+{
+	class IDashboardDataViewEntry;
+}
+
+namespace AudioModulationInsights
+{
+	class FModulationMatrixDashboardViewFactory : public UE::Audio::Insights::FTraceObjectTableDashboardViewFactory
+	{
+	public:
+		FModulationMatrixDashboardViewFactory();
+		virtual ~FModulationMatrixDashboardViewFactory();
+
+		virtual FName GetName() const override;
+		virtual FText GetDisplayName() const override;
+		virtual FSlateIcon GetIcon() const override;
+		virtual UE::Audio::Insights::EDefaultDashboardTabStack GetDefaultTabStack() const override;
+
+		virtual TSharedRef<SWidget> MakeWidget(TSharedRef<SDockTab> OwnerTab, const FSpawnTabArgs& SpawnTabArgs) override;
+
+		virtual void ProcessEntries(UE::Audio::Insights::FTraceTableDashboardViewFactory::EProcessReason Reason) override;
+
+		virtual TSharedPtr<SWidget> OnConstructContextMenu() override;
+		virtual FReply OnDataRowKeyInput(const FGeometry& InGeometry, const FKeyEvent& InKeyEvent) const override;
+
+	private:
+		virtual const TMap<FName, UE::Audio::Insights::FTraceTableDashboardViewFactory::FColumnData>& GetColumns() const override;
+
+		void CreateDefaultColumnData();
+		void RegisterDelegates();
+
+		void BindCommands();
+
+		TSharedRef<SWidget> MakeModulatingSourceTypeFilterWidget();
+
+		void OnControlBusesAdded(const FModulationMatrixTraceProvider::BusIdToBusInfoMap& AddedControlBuses);
+		void OnControlBusesRemoved(const TArray<FName>& RemovedControlBusesNames);
+
+#if WITH_EDITOR
+		void OnAudioDeviceDestroyed(::Audio::FDeviceId InDeviceId);
+#endif // WITH_EDITOR
+
+		void FilterByModulatingSourceName();
+		void FilterByModulatingSourceType();
+
+		virtual void SortTable() override;
+		virtual FSlateColor GetRowColor(const TSharedPtr<UE::Audio::Insights::IDashboardDataViewEntry>& InRowDataPtr) override;
+
+		enum class EModulatingSourceComboboxSelection : uint8
+		{
+			All,
+			BusMixes,
+			Generators
+		};
+
+		TSharedPtr<FModulationMatrixTraceProvider> ModulationMatrixTraceProvider;
+		TMap<FName, FTraceTableDashboardViewFactory::FColumnData> ModulationMatrixColumnData;
+		TArray<FName> ActiveBusNames;
+
+		FDelegateHandle OnControlBusAddedHandle;
+		FDelegateHandle OnControlBusRemovedHandle;
+
+#if WITH_EDITOR
+		FDelegateHandle OnDeviceDestroyedHandle;
+#endif // WITH_EDITOR
+
+		using FComboboxSelectionItem = TPair<EModulatingSourceComboboxSelection, FText>;
+		TArray<TSharedPtr<FComboboxSelectionItem>> ModulatingSourceTypes;
+		TSharedPtr<FComboboxSelectionItem> SelectedModulatingSourceType;
+
+		TSharedPtr<FUICommandList> CommandList;
+	};
+} // namespace AudioModulationInsights
diff --git a/Engine/Plugins/Runtime/AudioModulationInsights/Source/AudioModulationInsights/Public/AudioModulationInsightsModule.h b/Engine/Plugins/Runtime/AudioModulationInsights/Source/AudioModulationInsights/Public/AudioModulationInsightsModule.h
new file mode 100644
--- /dev/null
+++ b/Engine/Plugins/Runtime/AudioModulationInsights/Source/AudioModulationInsights/Public/AudioModulationInsightsModule.h
@@ -0,0 +1,18 @@
+// Copyright Epic Games, Inc. All Rights Reserved.
+#pragma once
+
+#include "Modules/ModuleInterface.h"
+
+#define UE_API AUDIOMODULATIONINSIGHTS_API
+
+class FAudioModulationInsightsModule : public IModuleInterface
+{
+public:
+	UE_API virtual void StartupModule() override;
+	UE_API virtual void ShutdownModule() override;
+
+private:
+	
+};
+
+#undef UE_API
