[AudioModulationInsights] Fixes from some detected BugHawk issues.

diff --git a/Engine/Plugins/Runtime/AudioModulationInsights/Source/AudioModulationInsights/Private/Messages/ModulationMatrixTraceMessages.h b/Engine/Plugins/Runtime/AudioModulationInsights/Source/AudioModulationInsights/Private/Messages/ModulationMatrixTraceMessages.h
--- a/Engine/Plugins/Runtime/AudioModulationInsights/Source/AudioModulationInsights/Private/Messages/ModulationMatrixTraceMessages.h
+++ b/Engine/Plugins/Runtime/AudioModulationInsights/Source/AudioModulationInsights/Private/Messages/ModulationMatrixTraceMessages.h
@@ -13,7 +13,6 @@ namespace AudioModulationInsights
 
 	enum class EModulationMatrixEntryType : uint8
 	{
-		// Arranged alphabetically to simplify sorting
 		BusMix,
 		BusFinalValues,
 		Generator,
diff --git a/Engine/Plugins/Runtime/AudioModulationInsights/Source/AudioModulationInsights/Private/Providers/ModulationMatrixTraceProvider.cpp b/Engine/Plugins/Runtime/AudioModulationInsights/Source/AudioModulationInsights/Private/Providers/ModulationMatrixTraceProvider.cpp
--- a/Engine/Plugins/Runtime/AudioModulationInsights/Source/AudioModulationInsights/Private/Providers/ModulationMatrixTraceProvider.cpp
+++ b/Engine/Plugins/Runtime/AudioModulationInsights/Source/AudioModulationInsights/Private/Providers/ModulationMatrixTraceProvider.cpp
@@ -112,7 +112,11 @@ namespace AudioModulationInsights
 		// Process messages
 		uint32 NumActiveBuses = ActiveControlBuses.Num();
 
-		if (!(TraceMessages.RegisterBusMessages.IsEmpty() && (TraceMessages.BusMixActivateMessages.IsEmpty() || TraceMessages.GeneratorActivateMessages.IsEmpty())))
+		const bool bHasRegisterBusMsgs       = !TraceMessages.RegisterBusMessages.IsEmpty();
+		const bool bHasBusMixActivateMsgs    = !TraceMessages.BusMixActivateMessages.IsEmpty();
+		const bool bHasGeneratorActivateMsgs = !TraceMessages.GeneratorActivateMessages.IsEmpty();
+
+		if (bHasRegisterBusMsgs && (bHasBusMixActivateMsgs || bHasGeneratorActivateMsgs))
 		{
 			TMap<FSourceId, TMap<FBusId, FString>> ModulatingSourceIdToBusesInfoMap;
 
diff --git a/Engine/Plugins/Runtime/AudioModulationInsights/Source/AudioModulationInsights/Private/Views/ControlBusDashboardViewFactory.cpp b/Engine/Plugins/Runtime/AudioModulationInsights/Source/AudioModulationInsights/Private/Views/ControlBusDashboardViewFactory.cpp
--- a/Engine/Plugins/Runtime/AudioModulationInsights/Source/AudioModulationInsights/Private/Views/ControlBusDashboardViewFactory.cpp
+++ b/Engine/Plugins/Runtime/AudioModulationInsights/Source/AudioModulationInsights/Private/Views/ControlBusDashboardViewFactory.cpp
@@ -211,7 +211,7 @@ namespace AudioModulationInsights
 				});
 			}
 		}
-		else if (SortByColumn == "ParamName")
+		else if (SortByColumn == "ParameterName")
 		{
 			if (SortMode == EColumnSortMode::Ascending)
 			{
