[AudioModulationInsights] FAudioModulationInsightsModule::ShutdownModule: Fixed IAudioInsightsEditorModule::IsModuleLoaded() check

diff --git a/Engine/Plugins/Runtime/AudioModulationInsights/Source/AudioModulationInsights/Private/AudioModulationInsightsModule.cpp b/Engine/Plugins/Runtime/AudioModulationInsights/Source/AudioModulationInsights/Private/AudioModulationInsightsModule.cpp
--- a/Engine/Plugins/Runtime/AudioModulationInsights/Source/AudioModulationInsights/Private/AudioModulationInsightsModule.cpp
+++ b/Engine/Plugins/Runtime/AudioModulationInsights/Source/AudioModulationInsights/Private/AudioModulationInsightsModule.cpp
@@ -28,14 +28,14 @@ void FAudioModulationInsightsModule::StartupModule()
 void FAudioModulationInsightsModule::ShutdownModule()
 {
 #if WITH_EDITOR
-	if (IAudioInsightsEditorModule::IsModuleLoaded())
+	if (!IAudioInsightsEditorModule::IsModuleLoaded())
 	{
 		return;
 	}
 
 	IAudioInsightsEditorModule& AudioInsightsModule = IAudioInsightsEditorModule::GetChecked();
 #else
-	if (IAudioInsightsModule::IsModuleLoaded())
+	if (!IAudioInsightsModule::IsModuleLoaded())
 	{
 		return;
 	}
