[AudioModulationInsights] FAudioModulationInsightsModule::ShutdownModule: wrap logic with an if (!IsRunningCommandlet())

diff --git a/Engine/Plugins/Runtime/AudioModulationInsights/Source/AudioModulationInsights/Private/AudioModulationInsightsModule.cpp b/Engine/Plugins/Runtime/AudioModulationInsights/Source/AudioModulationInsights/Private/AudioModulationInsightsModule.cpp
--- a/Engine/Plugins/Runtime/AudioModulationInsights/Source/AudioModulationInsights/Private/AudioModulationInsightsModule.cpp
+++ b/Engine/Plugins/Runtime/AudioModulationInsights/Source/AudioModulationInsights/Private/AudioModulationInsightsModule.cpp
@@ -27,24 +27,27 @@ void FAudioModulationInsightsModule::StartupModule()
 
 void FAudioModulationInsightsModule::ShutdownModule()
 {
-#if WITH_EDITOR
-	if (!IAudioInsightsEditorModule::IsModuleLoaded())
-	{
-		return;
-	}
+	if (!IsRunningCommandlet())
+	{	
+	#if WITH_EDITOR
+		if (!IAudioInsightsEditorModule::IsModuleLoaded())
+		{
+			return;
+		}
 
-	IAudioInsightsEditorModule& AudioInsightsModule = IAudioInsightsEditorModule::GetChecked();
-#else
-	if (!IAudioInsightsModule::IsModuleLoaded())
-	{
-		return;
-	}
+		IAudioInsightsEditorModule& AudioInsightsModule = IAudioInsightsEditorModule::GetChecked();
+	#else
+		if (!IAudioInsightsModule::IsModuleLoaded())
+		{
+			return;
+		}
 
-	IAudioInsightsModule& AudioInsightsModule = IAudioInsightsModule::GetChecked();
-#endif // WITH_EDITOR
-		
-	AudioInsightsModule.UnregisterDashboardViewFactory("ControlBuses");
-	AudioInsightsModule.UnregisterDashboardViewFactory("ModulationMatrix");
+		IAudioInsightsModule& AudioInsightsModule = IAudioInsightsModule::GetChecked();
+	#endif // WITH_EDITOR
+			
+		AudioInsightsModule.UnregisterDashboardViewFactory("ControlBuses");
+		AudioInsightsModule.UnregisterDashboardViewFactory("ModulationMatrix");
+	}
 }
 
 IMPLEMENT_MODULE(FAudioModulationInsightsModule, AudioModulationInsights);
