[AudioInsights]

* Remove obsolete split view in Control Buses dashboard
* Fixes: https://jira.it.epicgames.com/browse/UE-350379

Preflight: https://horde.devtools.epicgames.com/job/68fb640c206b31f290423f2d

diff --git a/Engine/Plugins/Runtime/AudioModulationInsights/Source/AudioModulationInsights/Private/Views/ControlBusDashboardViewFactory.cpp b/Engine/Plugins/Runtime/AudioModulationInsights/Source/AudioModulationInsights/Private/Views/ControlBusDashboardViewFactory.cpp
--- a/Engine/Plugins/Runtime/AudioModulationInsights/Source/AudioModulationInsights/Private/Views/ControlBusDashboardViewFactory.cpp
+++ b/Engine/Plugins/Runtime/AudioModulationInsights/Source/AudioModulationInsights/Private/Views/ControlBusDashboardViewFactory.cpp
@@ -84,26 +84,13 @@ namespace AudioModulationInsights
 	TSharedRef<SWidget> FControlBusDashboardViewFactory::MakeWidget(TSharedRef<SDockTab> OwnerTab, const FSpawnTabArgs& SpawnTabArgs)
 	{
 		TSharedRef<SWidget> TableDashboardWidget = UE::Audio::Insights::FTraceTableDashboardViewFactory::MakeWidget(OwnerTab, SpawnTabArgs);
-		TSharedRef<SWidget> BusWatchWidget = MakeControlBusWatchWidget();
 
 		return SNew(SHorizontalBox)
 			+ SHorizontalBox::Slot()
 			.HAlign(HAlign_Fill)
 			.Padding(0.0f, 6.0f, 0.0f, 0.0f)
 			[
-				// Dashboard and plots area
-				SNew(SSplitter)
-					.Orientation(Orient_Horizontal)
-					+ SSplitter::Slot()
-					.Value(0.55f)
-					[
-						TableDashboardWidget
-					]
-					+ SSplitter::Slot()
-					.Value(0.45f)
-					[
-						BusWatchWidget
-					]
+				TableDashboardWidget
 			];
 	}
 
@@ -259,16 +246,6 @@ namespace AudioModulationInsights
 		}
 	}
 
-	TSharedRef<SWidget> FControlBusDashboardViewFactory::MakeControlBusListWidget()
-	{
-		return SNew(SBox);
-	}
-
-	TSharedRef<SWidget> FControlBusDashboardViewFactory::MakeControlBusWatchWidget()
-	{
-		return SNew(SBox);
-	}
-
 	TSharedPtr<SWidget> FControlBusDashboardViewFactory::OnConstructContextMenu()
 	{
 		const FAudioModulationInsightsCommands& Commands = FAudioModulationInsightsCommands::Get();
diff --git a/Engine/Plugins/Runtime/AudioModulationInsights/Source/AudioModulationInsights/Private/Views/ControlBusDashboardViewFactory.h b/Engine/Plugins/Runtime/AudioModulationInsights/Source/AudioModulationInsights/Private/Views/ControlBusDashboardViewFactory.h
--- a/Engine/Plugins/Runtime/AudioModulationInsights/Source/AudioModulationInsights/Private/Views/ControlBusDashboardViewFactory.h
+++ b/Engine/Plugins/Runtime/AudioModulationInsights/Source/AudioModulationInsights/Private/Views/ControlBusDashboardViewFactory.h
@@ -24,9 +24,6 @@ namespace AudioModulationInsights
 
 		virtual void SortTable() override;
 
-		virtual TSharedRef<SWidget> MakeControlBusListWidget();
-		virtual TSharedRef<SWidget> MakeControlBusWatchWidget();
-
 		virtual TSharedPtr<SWidget> OnConstructContextMenu() override;
 		virtual FReply OnDataRowKeyInput(const FGeometry& InGeometry, const FKeyEvent& InKeyEvent) const override;
 
