[AudioInsights]

* ControlBus message support for AudioInsightsCacheManager
* ControBus messages now inherit from IAudioCachedMessage
* Tested AudioQAEditor AudioQA level, control buses still populate dashboard correctly
* https://jira.it.epicgames.com/browse/UE-314913

* Note: Support for OnTimingViewTimeMarkerChanged() will follow this, keeping separation of concerns

diff --git a/Engine/Plugins/Runtime/AudioModulationInsights/Source/AudioModulationInsights/Private/Messages/ControlBusTraceMessages.h b/Engine/Plugins/Runtime/AudioModulationInsights/Source/AudioModulationInsights/Private/Messages/ControlBusTraceMessages.h
--- a/Engine/Plugins/Runtime/AudioModulationInsights/Source/AudioModulationInsights/Private/Messages/ControlBusTraceMessages.h
+++ b/Engine/Plugins/Runtime/AudioModulationInsights/Source/AudioModulationInsights/Private/Messages/ControlBusTraceMessages.h
@@ -2,6 +2,7 @@
 #pragma once
 
 #include "AudioDefines.h"
+#include "Cache/IAudioCachedMessage.h"
 #include "Messages/AnalyzerMessageQueue.h"
 #include "Trace/Analyzer.h"
 #include "Views/TableDashboardViewFactory.h"
@@ -11,7 +12,14 @@ namespace AudioModulationInsights
 
 	using FBusId = uint32;
 
-	struct FControlBusMessageBase
+	namespace ControlBusMessageNames
+	{
+		inline const FName Activate(TEXT("ControlBusActivate"));
+		inline const FName Deactivate(TEXT("ControlBusDeactivate"));
+		inline const FName Update(TEXT("ControlBusUpdate"));
+	}
+
+	struct FControlBusMessageBase : public UE::Audio::Insights::IAudioCachedMessage
 	{
 		FControlBusMessageBase() = default;
 		FControlBusMessageBase(const UE::Trace::IAnalyzer::FOnEventContext& InContext)
@@ -20,31 +28,40 @@ namespace AudioModulationInsights
 			DeviceId = static_cast<::Audio::FDeviceId>(EventData.GetValue<uint32>("DeviceId"));
 			Timestamp = InContext.EventTime.AsSeconds(EventData.GetValue<uint64>("Timestamp"));
 			ControlBusId = static_cast<FBusId>(EventData.GetValue<uint32>("ControlBusId"));
+			EventData.GetString("Name", BusName);
+			EventData.GetString("ParamName", ParamName);
 		}
 
+		// from UE::Audio::Insights::IAudioCachedMessage
+		virtual uint64 GetID() const override { return ControlBusId; }
+		virtual uint32 GetSizeOf() const override 
+		{
+			return sizeof(FControlBusMessageBase) + BusName.GetAllocatedSize() + ParamName.GetAllocatedSize(); 
+		};
+
 		Audio::FDeviceId DeviceId = INDEX_NONE;
 		FBusId ControlBusId = INDEX_NONE;
-		double Timestamp = 0.0;
+		FString BusName;
+		FString ParamName;
 	};
 
-	struct FControlBusActivateMessage : public FControlBusMessageBase
+	struct FControlBusActivateMessage final : public FControlBusMessageBase
 	{
 		FControlBusActivateMessage() = default;
-		FControlBusActivateMessage(const UE::Trace::IAnalyzer::FOnEventContext& InContext)
-			: FControlBusMessageBase(InContext)
-		{
-			const UE::Trace::IAnalyzer::FEventData& EventData = InContext.EventData;
-			EventData.GetString("Name", BusName);
-			EventData.GetString("ParamName", ParamName);
-		}
-
-		FString BusName;
-		FString ParamName;
+		FControlBusActivateMessage(const UE::Trace::IAnalyzer::FOnEventContext& InContext) : FControlBusMessageBase(InContext) {}
+		// from UE::Audio::Insights::IAudioCachedMessage
+		virtual const FName GetMessageName() const override { return ControlBusMessageNames::Activate; };
 	};
 
-	using FControlBusDeactivateMessage = FControlBusMessageBase;
+	struct FControlBusDeactivateMessage final : public FControlBusMessageBase
+	{
+		FControlBusDeactivateMessage() = default;
+		FControlBusDeactivateMessage(const UE::Trace::IAnalyzer::FOnEventContext& InContext) : FControlBusMessageBase(InContext) {}
+		// from UE::Audio::Insights::IAudioCachedMessage
+		virtual const FName GetMessageName() const override { return ControlBusMessageNames::Deactivate; };
+	};
 
-	struct FControlBusUpdateMessage : public FControlBusMessageBase
+	struct FControlBusUpdateMessage final : public FControlBusMessageBase
 	{
 		FControlBusUpdateMessage() = default;
 		FControlBusUpdateMessage(const UE::Trace::IAnalyzer::FOnEventContext& InContext)
@@ -52,12 +69,15 @@ namespace AudioModulationInsights
 		{
 			const UE::Trace::IAnalyzer::FEventData& EventData = InContext.EventData;
 			Value = EventData.GetValue<float>("Value");
-			EventData.GetString("Name", BusName);
-			EventData.GetString("ParamName", ParamName);
 		}
 
-		FString BusName;
-		FString ParamName;
+		// from UE::Audio::Insights::IAudioCachedMessage
+		virtual uint32 GetSizeOf() const override 
+		{
+			return sizeof(FControlBusUpdateMessage) + BusName.GetAllocatedSize() + ParamName.GetAllocatedSize(); 
+		};
+		virtual const FName GetMessageName() const override { return ControlBusMessageNames::Update; };
+
 		float Value = 1.0f;
 	};
 
diff --git a/Engine/Plugins/Runtime/AudioModulationInsights/Source/AudioModulationInsights/Private/Providers/ControlBusTraceProvider.cpp b/Engine/Plugins/Runtime/AudioModulationInsights/Source/AudioModulationInsights/Private/Providers/ControlBusTraceProvider.cpp
--- a/Engine/Plugins/Runtime/AudioModulationInsights/Source/AudioModulationInsights/Private/Providers/ControlBusTraceProvider.cpp
+++ b/Engine/Plugins/Runtime/AudioModulationInsights/Source/AudioModulationInsights/Private/Providers/ControlBusTraceProvider.cpp
@@ -100,19 +100,19 @@ namespace AudioModulationInsights
 				{
 					case RouteId_Activate:
 					{
-						Messages.ActivateMessages.Enqueue(FControlBusActivateMessage { Context });
+						CacheMessage<FControlBusActivateMessage>(Context, Messages.ActivateMessages);
 						break;
 					}
 
 					case RouteId_Deactivate:
 					{
-						Messages.DeactivateMessages.Enqueue(FControlBusDeactivateMessage{ Context });
+						CacheMessage<FControlBusDeactivateMessage>(Context, Messages.DeactivateMessages);
 						break;
 					}
 
 					case RouteId_Update:
 					{
-						Messages.UpdateMessages.Enqueue(FControlBusUpdateMessage { Context });
+						CacheMessage<FControlBusUpdateMessage>(Context, Messages.UpdateMessages);
 						break;
 					}
 
