[AudioInsights]

* Support AudioInsightsCacheManager for Control Bus

* Consolidate FControlBusTraceProvider code into Reusable functions
* Abide by DRY principle, cleans up duplicated code previously
* And means OnTimingViewTimeMarkerChanged() can re-use these functions

Send a new ControlBusActivePing trace event every second
* Follows the is alive paradigm
* Means we can leverage the CacheManager's IterateTo function and then FindClosest function
* Instead of calling IterateTo with the ControlBusUpdate message which is sent much more frequently
* So a performance win, and reducing memory footprint amount within a time window (since we move Name and ParamName into ControlBusActivePing)

Removing Name and ParamName from the frequently sent ControlBusUpdate trace event
* Discussed with Steve, sending string data this frequently is not required
* We instead can rely on the name and param name inside the ControlBusActivePing trace event
* Bonus side effect: FControlBusMessageBase no longer needs Name and ParamName

Passed El BugHawk
Fixes: https://jira.it.epicgames.com/browse/UE-314913

diff --git a/Engine/Plugins/Runtime/AudioModulation/Source/AudioModulation/Private/AudioModulationSystem.cpp b/Engine/Plugins/Runtime/AudioModulation/Source/AudioModulation/Private/AudioModulationSystem.cpp
--- a/Engine/Plugins/Runtime/AudioModulation/Source/AudioModulation/Private/AudioModulationSystem.cpp
+++ b/Engine/Plugins/Runtime/AudioModulation/Source/AudioModulation/Private/AudioModulationSystem.cpp
@@ -38,12 +38,18 @@ DECLARE_DWORD_COUNTER_STAT(TEXT("Render Queue Commands Processed"), STAT_AudioMo
 
 #if UE_AUDIO_PROFILERTRACE_ENABLED
 UE_TRACE_EVENT_BEGIN(Audio, ControlBusUpdate)
+	UE_TRACE_EVENT_FIELD(uint32, DeviceId)
+	UE_TRACE_EVENT_FIELD(uint32, ControlBusId)
+	UE_TRACE_EVENT_FIELD(uint64, Timestamp)
+	UE_TRACE_EVENT_FIELD(float, Value)
+UE_TRACE_EVENT_END()
+
+UE_TRACE_EVENT_BEGIN(Audio, ControlBusActivePing)
 	UE_TRACE_EVENT_FIELD(uint32, DeviceId)
 	UE_TRACE_EVENT_FIELD(uint32, ControlBusId)
 	UE_TRACE_EVENT_FIELD(uint64, Timestamp)
 	UE_TRACE_EVENT_FIELD(UE::Trace::WideString, Name)
 	UE_TRACE_EVENT_FIELD(UE::Trace::WideString, ParamName)
-	UE_TRACE_EVENT_FIELD(float, Value)
 UE_TRACE_EVENT_END()
 
 UE_TRACE_EVENT_BEGIN(Audio, BusMixRegisterBus)
@@ -664,15 +670,30 @@ namespace AudioModulation
 			BusIdsFinalValues.Reserve(NumBuses);
 			BusFinalValues.Reserve(NumBuses);
 
+			ControlBusActivePingTimeSec -= InElapsed;
+			const bool bSendControlBusActivePingTrace = (ControlBusActivePingTimeSec <= 0.0f);
+			if (bSendControlBusActivePingTrace)
+			{
+				ControlBusActivePingTimeSec = ControlBusActivePingIntervalTimeSec;
+			}
+
 			// Control Bus message trace
 			for (const auto& [BusId, ControlBusProxy] : RefProxies.Buses)
 			{
+				if (bSendControlBusActivePingTrace)
+				{
+					UE_TRACE_LOG(Audio, ControlBusActivePing, AudioChannel)
+						<< ControlBusActivePing.DeviceId(static_cast<uint32>(AudioDeviceId))
+						<< ControlBusActivePing.ControlBusId(static_cast<uint32>(BusId))
+						<< ControlBusActivePing.Timestamp(FPlatformTime::Cycles64())
+						<< ControlBusActivePing.Name(*(ControlBusProxy.GetPath().ToString()))
+						<< ControlBusActivePing.ParamName(*(ControlBusProxy.GetParameterName().ToString()));
+				}
+
 				UE_TRACE_LOG(Audio, ControlBusUpdate, AudioChannel)
 					<< ControlBusUpdate.DeviceId(static_cast<uint32>(AudioDeviceId))
 					<< ControlBusUpdate.ControlBusId(static_cast<uint32>(BusId))
 					<< ControlBusUpdate.Timestamp(FPlatformTime::Cycles64())
-					<< ControlBusUpdate.Name(*(ControlBusProxy.GetPath().ToString()))
-					<< ControlBusUpdate.ParamName(*(ControlBusProxy.GetParameterName().ToString()))
 					<< ControlBusUpdate.Value(ControlBusProxy.GetValue());
 
 				// Collect the final value of this control bus
diff --git a/Engine/Plugins/Runtime/AudioModulation/Source/AudioModulation/Private/AudioModulationSystem.h b/Engine/Plugins/Runtime/AudioModulation/Source/AudioModulation/Private/AudioModulationSystem.h
--- a/Engine/Plugins/Runtime/AudioModulation/Source/AudioModulation/Private/AudioModulationSystem.h
+++ b/Engine/Plugins/Runtime/AudioModulation/Source/AudioModulation/Private/AudioModulationSystem.h
@@ -276,6 +276,11 @@ namespace AudioModulation
 		using FParameterRegistry = TMap<FName, Audio::FModulationParameter>;
 		FParameterRegistry ParameterRegistry;
 
+#if UE_AUDIO_PROFILERTRACE_ENABLED
+		float ControlBusActivePingTimeSec = 0.0f;
+		static constexpr float ControlBusActivePingIntervalTimeSec = 1.0f;
+#endif // UE_AUDIO_PROFILERTRACE_ENABLED
+
 #if !UE_BUILD_SHIPPING
 	public:
 		void SetDebugBusFilter(const FString* InFilter);
diff --git a/Engine/Plugins/Runtime/AudioModulationInsights/Source/AudioModulationInsights/Private/Messages/ControlBusTraceMessages.h b/Engine/Plugins/Runtime/AudioModulationInsights/Source/AudioModulationInsights/Private/Messages/ControlBusTraceMessages.h
--- a/Engine/Plugins/Runtime/AudioModulationInsights/Source/AudioModulationInsights/Private/Messages/ControlBusTraceMessages.h
+++ b/Engine/Plugins/Runtime/AudioModulationInsights/Source/AudioModulationInsights/Private/Messages/ControlBusTraceMessages.h
@@ -28,35 +28,39 @@ namespace AudioModulationInsights
 			DeviceId = static_cast<::Audio::FDeviceId>(EventData.GetValue<uint32>("DeviceId"));
 			Timestamp = InContext.EventTime.AsSeconds(EventData.GetValue<uint64>("Timestamp"));
 			ControlBusId = static_cast<FBusId>(EventData.GetValue<uint32>("ControlBusId"));
-			EventData.GetString("Name", BusName);
-			EventData.GetString("ParamName", ParamName);
 		}
 
 		// from UE::Audio::Insights::IAudioCachedMessage
 		virtual uint64 GetID() const override { return ControlBusId; }
-		virtual uint32 GetSizeOf() const override 
-		{
-			return sizeof(FControlBusMessageBase) + BusName.GetAllocatedSize() + ParamName.GetAllocatedSize(); 
-		};
+		virtual uint32 GetSizeOf() const override { return sizeof(FControlBusMessageBase); };
 
 		Audio::FDeviceId DeviceId = INDEX_NONE;
 		FBusId ControlBusId = INDEX_NONE;
-		FString BusName;
-		FString ParamName;
 	};
 
 	struct FControlBusActivateMessage final : public FControlBusMessageBase
 	{
 		FControlBusActivateMessage() = default;
-		FControlBusActivateMessage(const UE::Trace::IAnalyzer::FOnEventContext& InContext) : FControlBusMessageBase(InContext) {}
+		FControlBusActivateMessage(const UE::Trace::IAnalyzer::FOnEventContext& InContext) : FControlBusMessageBase(InContext) 
+		{
+			const UE::Trace::IAnalyzer::FEventData& EventData = InContext.EventData;
+			EventData.GetString("Name", BusName);
+			EventData.GetString("ParamName", ParamName);
+		}
+
 		// from UE::Audio::Insights::IAudioCachedMessage
 		virtual const FName GetMessageName() const override { return ControlBusMessageNames::Activate; };
+		virtual uint32 GetSizeOf() const override { return sizeof(FControlBusActivateMessage) + BusName.GetAllocatedSize() + ParamName.GetAllocatedSize(); };
+
+		FString BusName;
+		FString ParamName;
 	};
 
 	struct FControlBusDeactivateMessage final : public FControlBusMessageBase
 	{
 		FControlBusDeactivateMessage() = default;
 		FControlBusDeactivateMessage(const UE::Trace::IAnalyzer::FOnEventContext& InContext) : FControlBusMessageBase(InContext) {}
+
 		// from UE::Audio::Insights::IAudioCachedMessage
 		virtual const FName GetMessageName() const override { return ControlBusMessageNames::Deactivate; };
 	};
@@ -72,16 +76,13 @@ namespace AudioModulationInsights
 		}
 
 		// from UE::Audio::Insights::IAudioCachedMessage
-		virtual uint32 GetSizeOf() const override 
-		{
-			return sizeof(FControlBusUpdateMessage) + BusName.GetAllocatedSize() + ParamName.GetAllocatedSize(); 
-		};
+		virtual uint32 GetSizeOf() const override { return sizeof(FControlBusUpdateMessage); };
 		virtual const FName GetMessageName() const override { return ControlBusMessageNames::Update; };
 
 		float Value = 1.0f;
 	};
 
-	class FControlBusDashboardEntry : public UE::Audio::Insights::FSoundAssetDashboardEntry
+	class FControlBusDashboardEntry final : public UE::Audio::Insights::FSoundAssetDashboardEntry
 	{
 	public:
 		FControlBusDashboardEntry() = default;
diff --git a/Engine/Plugins/Runtime/AudioModulationInsights/Source/AudioModulationInsights/Private/Providers/ControlBusTraceProvider.cpp b/Engine/Plugins/Runtime/AudioModulationInsights/Source/AudioModulationInsights/Private/Providers/ControlBusTraceProvider.cpp
--- a/Engine/Plugins/Runtime/AudioModulationInsights/Source/AudioModulationInsights/Private/Providers/ControlBusTraceProvider.cpp
+++ b/Engine/Plugins/Runtime/AudioModulationInsights/Source/AudioModulationInsights/Private/Providers/ControlBusTraceProvider.cpp
@@ -2,6 +2,8 @@
 #include "ControlBusTraceProvider.h"
 
 #include "AudioInsightsUtils.h"
+#include "Cache/AudioInsightsCacheManager.h"
+#include "IAudioInsightsModule.h"
 #include "Trace/Analyzer.h"
 #include "TraceServices/Model/AnalysisSession.h"
 
@@ -12,60 +14,98 @@ namespace AudioModulationInsights
 		return "ControlBusProvider";
 	}
 
-	bool FControlBusTraceProvider::ProcessMessages()
+	TSharedPtr<FControlBusDashboardEntry>* FControlBusTraceProvider::GetEntry(const FControlBusMessageBase& InControlBusMessage)
 	{
-		using namespace UE::Audio::Insights;
+		TSharedPtr<FControlBusDashboardEntry>& FoundEntry = FindOrAddDeviceEntry(InControlBusMessage.DeviceId, InControlBusMessage.ControlBusId);
 
-		auto BumpEntryFunc = [this](const FControlBusMessageBase& Msg)
+		if (!FoundEntry.IsValid())
 		{
-			TSharedPtr<FControlBusDashboardEntry>* ToReturn = nullptr;
-			UpdateDeviceEntry(Msg.DeviceId, Msg.ControlBusId, [&ToReturn, &Msg](TSharedPtr<FControlBusDashboardEntry>& Entry)
-				{
-					if (!Entry.IsValid())
-					{
-						Entry = MakeShared<FControlBusDashboardEntry>();
-						Entry->DeviceId = Msg.DeviceId;
-						Entry->ControlBusId = Msg.ControlBusId;
-					}
-					ToReturn = &Entry;
-				});
+			FoundEntry = MakeShared<FControlBusDashboardEntry>();
+			FoundEntry->DeviceId = InControlBusMessage.DeviceId;
+			FoundEntry->ControlBusId = InControlBusMessage.ControlBusId;
+		}
 
-			return ToReturn;
-		};
+		return &FoundEntry;
+	}
 
-		ProcessMessageQueue<FControlBusActivateMessage>(TraceMessages.ActivateMessages, BumpEntryFunc,
-		[](const FControlBusActivateMessage& Msg, TSharedPtr<FControlBusDashboardEntry>* OutEntry)
+	void FControlBusTraceProvider::ProcessActivateMessage(const FControlBusActivateMessage& InControlBusMessage, const TSharedPtr<FControlBusDashboardEntry>* const OutControlBusDashboardEntry)
+	{
+		if (!OutControlBusDashboardEntry || !OutControlBusDashboardEntry->IsValid())
 		{
-			FControlBusDashboardEntry& EntryRef = *OutEntry->Get();
-			EntryRef.Name = *Msg.BusName;
-			EntryRef.ControlBusId = Msg.ControlBusId;
-			EntryRef.DisplayName = AudioInsightsUtils::ResolveObjectDisplayName(EntryRef.Name);
-			EntryRef.ParamName = Msg.ParamName;
-		});
+			return;
+		}
 
-		ProcessMessageQueue<FControlBusUpdateMessage>(TraceMessages.UpdateMessages, BumpEntryFunc,
-		[](const FControlBusUpdateMessage& Msg, TSharedPtr<FControlBusDashboardEntry>* OutEntry)
-		{
-			FControlBusDashboardEntry& EntryRef = *OutEntry->Get();
-			EntryRef.Name = *Msg.BusName;
-			EntryRef.DisplayName = AudioInsightsUtils::ResolveObjectDisplayName(EntryRef.Name);
-			EntryRef.ParamName = Msg.ParamName;
-			EntryRef.Value = Msg.Value;
-		});
+		FControlBusDashboardEntry& ControlBusEntry = *OutControlBusDashboardEntry->Get();
+
+		ControlBusEntry.Timestamp = InControlBusMessage.Timestamp;
+		ControlBusEntry.DeviceId = InControlBusMessage.DeviceId;
+		ControlBusEntry.ControlBusId = InControlBusMessage.ControlBusId;
+
+		ensure((!InControlBusMessage.BusName.IsEmpty()));
+		ensure((!InControlBusMessage.ParamName.IsEmpty()));
+
+		ControlBusEntry.Name = InControlBusMessage.BusName;
+		ControlBusEntry.DisplayName = UE::Audio::Insights::AudioInsightsUtils::ResolveObjectDisplayName(ControlBusEntry.Name);
+		ControlBusEntry.ParamName = InControlBusMessage.ParamName;
+	}
 
-		auto GetEntry = [this](const FControlBusMessageBase& Msg)
+	void FControlBusTraceProvider::ProcessUpdateMessage(const FControlBusUpdateMessage& InControlBusMessage, const TSharedPtr<FControlBusDashboardEntry>* const OutControlBusDashboardEntry)
+	{
+		if (!OutControlBusDashboardEntry || !OutControlBusDashboardEntry->IsValid())
 		{
-			return FindDeviceEntry(Msg.DeviceId, Msg.ControlBusId);
-		};
+			return;
+		}
+
+		// DeviceId and ControlBusId has already been set by ProcessActivateMessage, so we only need to update the timestamp and value here
+		FControlBusDashboardEntry& ControlBusEntry = *OutControlBusDashboardEntry->Get();
+		ControlBusEntry.Timestamp = InControlBusMessage.Timestamp;
+		ControlBusEntry.Value = InControlBusMessage.Value;
+	}
 
-		ProcessMessageQueue<FControlBusDeactivateMessage>(TraceMessages.DeactivateMessages, GetEntry,
-		[this](const FControlBusDeactivateMessage& Msg, TSharedPtr<FControlBusDashboardEntry>* OutEntry)
+	void FControlBusTraceProvider::ProcessDeactivateMessage(const FControlBusDeactivateMessage& InControlBusMessage, const TSharedPtr<FControlBusDashboardEntry>* const OutControlBusDashboardEntry)
+	{
+		if (!OutControlBusDashboardEntry || !OutControlBusDashboardEntry->IsValid())
+		{
+			return;
+		}
+	
+		if (OutControlBusDashboardEntry->Get()->Timestamp < InControlBusMessage.Timestamp)
 		{
-			if (OutEntry && (*OutEntry)->Timestamp < Msg.Timestamp)
+			RemoveDeviceEntry(InControlBusMessage.DeviceId, InControlBusMessage.ControlBusId);
+		}
+	}
+
+	bool FControlBusTraceProvider::ProcessMessages()
+	{
+		ProcessMessageQueue<FControlBusActivateMessage>(TraceMessages.ActivateMessages,
+			[this](const FControlBusMessageBase& Msg) 
+			{ 
+				return GetEntry(Msg);
+			},
+			[this](const FControlBusActivateMessage& Msg, TSharedPtr<FControlBusDashboardEntry>* const OutEntry)
 			{
-				RemoveDeviceEntry(Msg.DeviceId, Msg.ControlBusId);
-			}
-		});
+				ProcessActivateMessage(Msg, OutEntry);
+			});
+
+		ProcessMessageQueue<FControlBusDeactivateMessage>(TraceMessages.DeactivateMessages, 
+			[this](const FControlBusMessageBase& Msg)
+			{
+				return FindDeviceEntry(Msg.DeviceId, Msg.ControlBusId);
+			},
+			[this](const FControlBusDeactivateMessage& Msg, TSharedPtr<FControlBusDashboardEntry>* OutEntry)
+			{
+				ProcessDeactivateMessage(Msg, OutEntry);
+			});
+
+		ProcessMessageQueue<FControlBusUpdateMessage>(TraceMessages.UpdateMessages,
+			[this](const FControlBusMessageBase& Msg)
+			{
+				return FindDeviceEntry(Msg.DeviceId, Msg.ControlBusId);
+			},
+			[this](const FControlBusUpdateMessage& Msg, TSharedPtr<FControlBusDashboardEntry>* const OutEntry)
+			{
+				ProcessUpdateMessage(Msg, OutEntry);
+			});
 
 		return true;
 	}
@@ -87,6 +127,7 @@ namespace AudioModulationInsights
 
 				UE::Trace::IAnalyzer::FInterfaceBuilder& Builder = Context.InterfaceBuilder;
 				Builder.RouteEvent(RouteId_Activate, "Audio", "ControlBusActivate");
+				Builder.RouteEvent(RouteId_ActivePing, "Audio", "ControlBusActivePing");
 				Builder.RouteEvent(RouteId_Deactivate, "Audio", "ControlBusDeactivate");
 				Builder.RouteEvent(RouteId_Update, "Audio", "ControlBusUpdate");
 			}
@@ -99,6 +140,7 @@ namespace AudioModulationInsights
 				switch (RouteId)
 				{
 					case RouteId_Activate:
+					case RouteId_ActivePing:
 					{
 						CacheMessage<FControlBusActivateMessage>(Context, Messages.ActivateMessages);
 						break;
@@ -136,6 +178,7 @@ namespace AudioModulationInsights
 			enum : uint16
 			{
 				RouteId_Activate,
+				RouteId_ActivePing,
 				RouteId_Deactivate,
 				RouteId_Update
 			};
@@ -145,4 +188,48 @@ namespace AudioModulationInsights
 
 		return new FControlBusTraceAnalyzer(AsShared(), InSession);
 	}
+
+	void FControlBusTraceProvider::OnTimingViewTimeMarkerChanged(const double InTimeMarker)
+	{
+		DeviceDataMap.Empty();
+
+		const UE::Audio::Insights::FAudioInsightsCacheManager& AudioInsightsCacheManager = IAudioInsightsModule::GetChecked().GetCacheManager();
+
+		AudioInsightsCacheManager.IterateTo<FControlBusActivateMessage>(ControlBusMessageNames::Activate, InTimeMarker, [this](const FControlBusActivateMessage& InControlBusMessage)
+		{
+			TSharedPtr<FControlBusDashboardEntry>* const ControlBusDashboardEntry = GetEntry(InControlBusMessage);
+			ProcessActivateMessage(InControlBusMessage, ControlBusDashboardEntry);
+		});
+
+		AudioInsightsCacheManager.IterateTo<FControlBusDeactivateMessage>(ControlBusMessageNames::Deactivate, InTimeMarker, [this, InTimeMarker](const FControlBusDeactivateMessage& InControlBusMessage)
+		{
+			TSharedPtr<FControlBusDashboardEntry>* const ControlBusDashboardEntry = FindDeviceEntry(InControlBusMessage.DeviceId, InControlBusMessage.ControlBusId);
+			ProcessDeactivateMessage(InControlBusMessage, ControlBusDashboardEntry);
+		});
+
+		const FDeviceData* const DeviceData = FindFilteredDeviceData();
+		if (DeviceData)
+		{
+			for (auto& DeviceEntry : *DeviceData)
+			{
+				// Find the update message associated with this control bus
+				const FControlBusUpdateMessage* const ControlBusUpdateMessage = AudioInsightsCacheManager.FindClosestMessage<FControlBusUpdateMessage>(ControlBusMessageNames::Update, InTimeMarker, DeviceEntry.Key);
+				if (ControlBusUpdateMessage)
+				{
+					ProcessUpdateMessage(*ControlBusUpdateMessage, &DeviceEntry.Value);
+				}
+			}
+		}
+
+		// Call parent method to update LastMessageId
+		FTraceProviderBase::OnTimingViewTimeMarkerChanged(InTimeMarker);
+	}
+
+	void FControlBusTraceProvider::OnTimeControlMethodReset()
+	{
+		// When leaving the cache here, we intentionally clear entries, to avoid deactivated runtime control buses staying in the dashboard
+		// Since control buses update regularly, any active runtime control buses will be added back in within ProcessMessages()
+		Reset();
+	}
+
 } // namespace AudioModulationInsights
diff --git a/Engine/Plugins/Runtime/AudioModulationInsights/Source/AudioModulationInsights/Private/Providers/ControlBusTraceProvider.h b/Engine/Plugins/Runtime/AudioModulationInsights/Source/AudioModulationInsights/Private/Providers/ControlBusTraceProvider.h
--- a/Engine/Plugins/Runtime/AudioModulationInsights/Source/AudioModulationInsights/Private/Providers/ControlBusTraceProvider.h
+++ b/Engine/Plugins/Runtime/AudioModulationInsights/Source/AudioModulationInsights/Private/Providers/ControlBusTraceProvider.h
@@ -25,6 +25,15 @@ namespace AudioModulationInsights
 		static FName GetName_Static();
 
 	private:
+		// from UE::Audio::Insights::FTraceProviderBase via UE::Audio::Insights::TDeviceDataMapTraceProvider
+		virtual void OnTimingViewTimeMarkerChanged(const double InTimeMarker) override;
+		virtual void OnTimeControlMethodReset() override;
+
+		TSharedPtr<FControlBusDashboardEntry>* GetEntry(const FControlBusMessageBase& InControlBusMessage);
+
+		void ProcessActivateMessage(const FControlBusActivateMessage& InControlBusMessage, const TSharedPtr<FControlBusDashboardEntry>* const OutControlBusDashboardEntry);
+		void ProcessUpdateMessage(const FControlBusUpdateMessage& InControlBusMessage, const TSharedPtr<FControlBusDashboardEntry>* const OutControlBusDashboardEntry);
+		void ProcessDeactivateMessage(const FControlBusDeactivateMessage& InControlBusMessage, const TSharedPtr<FControlBusDashboardEntry>* const OutControlBusDashboardEntry);
 
 		FControlBusMessages TraceMessages;
 	};
