[AudioInsights] Support modulation messages with AudioInsightsCacheManager Pt.1

* ModulationMatrix messages now inherit from IAudioCachedMessage
* Tested AudioQAEditor AudioQA level, control buses still populate dashboard correctly
* https://jira.it.epicgames.com/browse/UE-314914

#Synced-CL 48281426

diff --git a/Engine/Plugins/Runtime/AudioModulationInsights/Source/AudioModulationInsights/Private/Messages/ModulationMatrixTraceMessages.h b/Engine/Plugins/Runtime/AudioModulationInsights/Source/AudioModulationInsights/Private/Messages/ModulationMatrixTraceMessages.h
--- a/Engine/Plugins/Runtime/AudioModulationInsights/Source/AudioModulationInsights/Private/Messages/ModulationMatrixTraceMessages.h
+++ b/Engine/Plugins/Runtime/AudioModulationInsights/Source/AudioModulationInsights/Private/Messages/ModulationMatrixTraceMessages.h
@@ -2,6 +2,7 @@
 #pragma once
 
 #include "AudioDefines.h"
+#include "Cache/IAudioCachedMessage.h"
 #include "Messages/AnalyzerMessageQueue.h"
 #include "Trace/Analyzer.h"
 #include "Views/TableDashboardViewFactory.h"
@@ -19,6 +20,16 @@ namespace AudioModulationInsights
 		None
 	};
 
+	namespace ModulationMatrixMessageNames
+	{
+		inline const FName BusMixRegisterBus(TEXT("ModulationMatrixBusMixRegisterBus"));
+		inline const FName BusMixActivate(TEXT("ModulationMatrixBusMixActivate"));
+		inline const FName BusMixUpdate(TEXT("ModulationMatrixBusMixUpdate"));
+		inline const FName GeneratorActivate(TEXT("ModulationMatrixGeneratorActivate"));
+		inline const FName GeneratorUpdate(TEXT("ModulationMatrixGeneratorUpdate"));
+		inline const FName BusFinalValuesUpdate(TEXT("ModulationMatrixBusFinalValuesUpdate"));
+		inline const FName ModulatingSourceDeactivate(TEXT("ModulationMatrixModulatingSourceDeactivate"));
+	}
 
 	// Modulation matrix dashboard entry
 	class FModulationMatrixDashboardEntry : public UE::Audio::Insights::FSoundAssetDashboardEntry
@@ -36,7 +47,7 @@ namespace AudioModulationInsights
 	
 
 	// Modulation matrix message base
-	struct FModulationMatrixMessageBase
+	struct FModulationMatrixMessageBase : public UE::Audio::Insights::IAudioCachedMessage
 	{
 		FModulationMatrixMessageBase() = default;
 
@@ -49,14 +60,16 @@ namespace AudioModulationInsights
 			SourceId  = static_cast<FSourceId>(EventData.GetValue<uint32>("SourceId"));
 		}
 
+		// from UE::Audio::Insights::IAudioCachedMessage
+		virtual uint64 GetID() const override { return SourceId; }
+
 		Audio::FDeviceId DeviceId = INDEX_NONE;
 		FSourceId SourceId = INDEX_NONE;
-		double Timestamp = 0.0;
 	};
 
 
 	// Register bus message
-	struct FModulationMatrixRegisterBusMessage : public FModulationMatrixMessageBase
+	struct FModulationMatrixRegisterBusMessage final : public FModulationMatrixMessageBase
 	{
 		FModulationMatrixRegisterBusMessage() = default;
 
@@ -69,17 +82,21 @@ namespace AudioModulationInsights
 			EventData.GetString("BusName", BusName);
 		}
 
+		// from UE::Audio::Insights::IAudioCachedMessage
+		virtual const FName GetMessageName() const override { return ModulationMatrixMessageNames::BusMixRegisterBus; }
+		virtual uint32 GetSizeOf() const override { return sizeof(FModulationMatrixRegisterBusMessage) + BusName.GetAllocatedSize(); }
+
 		FSourceId ModulatingSourceId = INDEX_NONE;
 		FString BusName;
 	};
 
 
 	// Activate messages
-	struct FModulationMatrixActivateMessage : public FModulationMatrixMessageBase
+	struct FModulationMatrixActivateMessageBase : public FModulationMatrixMessageBase
 	{
-		FModulationMatrixActivateMessage() = default;
+		FModulationMatrixActivateMessageBase() = default;
 
-		FModulationMatrixActivateMessage(const UE::Trace::IAnalyzer::FOnEventContext& InContext)
+		FModulationMatrixActivateMessageBase(const UE::Trace::IAnalyzer::FOnEventContext& InContext)
 			: FModulationMatrixMessageBase(InContext)
 		{
 			const UE::Trace::IAnalyzer::FEventData& EventData = InContext.EventData;
@@ -91,82 +108,116 @@ namespace AudioModulationInsights
 		FString Name;
 	};
 
-	struct FBusMixActivateMessage : public FModulationMatrixActivateMessage
+	struct FBusMixActivateMessage final : public FModulationMatrixActivateMessageBase
 	{
 		FBusMixActivateMessage() = default;
 
 		FBusMixActivateMessage(const UE::Trace::IAnalyzer::FOnEventContext& InContext)
-			: FModulationMatrixActivateMessage(InContext)
+			: FModulationMatrixActivateMessageBase(InContext)
 		{
 			EntryType = EModulationMatrixEntryType::BusMix;
 		}
+
+		// from UE::Audio::Insights::IAudioCachedMessage
+		virtual const FName GetMessageName() const override { return ModulationMatrixMessageNames::BusMixActivate; }
+		virtual uint32 GetSizeOf() const override { return sizeof(FBusMixActivateMessage) + Name.GetAllocatedSize(); }
 	};
 
-	struct FGeneratorActivateMessage : public FModulationMatrixActivateMessage
+	struct FGeneratorActivateMessage final : public FModulationMatrixActivateMessageBase
 	{
 		FGeneratorActivateMessage() = default;
 
 		FGeneratorActivateMessage(const UE::Trace::IAnalyzer::FOnEventContext& InContext)
-			: FModulationMatrixActivateMessage(InContext)
+			: FModulationMatrixActivateMessageBase(InContext)
 		{
 			EntryType = EModulationMatrixEntryType::Generator;
 		}
-	};
 
+		// from UE::Audio::Insights::IAudioCachedMessage
+		virtual const FName GetMessageName() const override { return ModulationMatrixMessageNames::GeneratorActivate; }
+		virtual uint32 GetSizeOf() const override { return sizeof(FGeneratorActivateMessage) + Name.GetAllocatedSize(); }
+	};
 
 	// Update messages
-	struct FModulationMatrixUpdateMessage : public FModulationMatrixMessageBase
+	struct FGeneratorUpdateMessage final : public FModulationMatrixMessageBase
 	{
-		FModulationMatrixUpdateMessage() = default;
+		FGeneratorUpdateMessage() = default;
 
-		FModulationMatrixUpdateMessage(const UE::Trace::IAnalyzer::FOnEventContext& InContext)
+		FGeneratorUpdateMessage(const UE::Trace::IAnalyzer::FOnEventContext& InContext)
 			: FModulationMatrixMessageBase(InContext)
 		{
 			const UE::Trace::IAnalyzer::FEventData& EventData = InContext.EventData;
 
-			const TArrayView<const FBusId> BusIds   = EventData.GetArrayView<FBusId>("BusIds");
-			const TArrayView<const float> BusValues = EventData.GetArrayView<float>("BusValues");
-
-			check(BusIds.Num() == BusValues.Num());
+			const TArrayView<const FBusId> BusIds = EventData.GetArrayView<FBusId>("BusIds");
+			const float GeneratorValue = EventData.GetValue<float>("GeneratorValue");
 
+			// All control buses use the same generator value
 			for (int32 Index = 0; Index < BusIds.Num(); ++Index)
 			{
-				BusIdToValueMap.Emplace(BusIds[Index], BusValues[Index]);
+				BusIdToValueMap.Emplace(BusIds[Index], GeneratorValue);
 			}
 		}
 
+		// from UE::Audio::Insights::IAudioCachedMessage
+		virtual const FName GetMessageName() const override { return ModulationMatrixMessageNames::GeneratorUpdate; }
+		virtual uint32 GetSizeOf() const override { return sizeof(FGeneratorUpdateMessage) + BusIdToValueMap.GetAllocatedSize(); }
+
 		TMap<FBusId, float> BusIdToValueMap;
 	};
 
-	using FBusMixUpdateMessage = FModulationMatrixUpdateMessage;
-
-	struct FGeneratorUpdateMessage : public FModulationMatrixMessageBase
+	struct FModulationMatrixUpdateMessageBase : public FModulationMatrixMessageBase
 	{
-		FGeneratorUpdateMessage() = default;
+		FModulationMatrixUpdateMessageBase() = default;
 
-		FGeneratorUpdateMessage(const UE::Trace::IAnalyzer::FOnEventContext& InContext)
+		FModulationMatrixUpdateMessageBase(const UE::Trace::IAnalyzer::FOnEventContext& InContext)
 			: FModulationMatrixMessageBase(InContext)
 		{
 			const UE::Trace::IAnalyzer::FEventData& EventData = InContext.EventData;
 
-			const TArrayView<const FBusId> BusIds = EventData.GetArrayView<FBusId>("BusIds");
-			const float GeneratorValue = EventData.GetValue<float>("GeneratorValue");
+			const TArrayView<const FBusId> BusIds   = EventData.GetArrayView<FBusId>("BusIds");
+			const TArrayView<const float> BusValues = EventData.GetArrayView<float>("BusValues");
+
+			check(BusIds.Num() == BusValues.Num());
 
-			// All control buses use the same generator value
 			for (int32 Index = 0; Index < BusIds.Num(); ++Index)
 			{
-				BusIdToValueMap.Emplace(BusIds[Index], GeneratorValue);
+				BusIdToValueMap.Emplace(BusIds[Index], BusValues[Index]);
 			}
 		}
 
 		TMap<FBusId, float> BusIdToValueMap;
 	};
 
-	using FBusFinalValuesUpdateMessage = FModulationMatrixUpdateMessage;
+	struct FBusMixUpdateMessage final : public FModulationMatrixUpdateMessageBase
+	{
+		FBusMixUpdateMessage() = default;
+		FBusMixUpdateMessage(const UE::Trace::IAnalyzer::FOnEventContext& InContext) : FModulationMatrixUpdateMessageBase(InContext) {}
 
+		// from UE::Audio::Insights::IAudioCachedMessage
+		virtual const FName GetMessageName() const override { return ModulationMatrixMessageNames::BusMixUpdate; }
+		virtual uint32 GetSizeOf() const override { return sizeof(FBusMixUpdateMessage) + BusIdToValueMap.GetAllocatedSize(); }
+	};
+
+	struct FBusFinalValuesUpdateMessage final : public FModulationMatrixUpdateMessageBase
+	{
+		FBusFinalValuesUpdateMessage() = default;
+		FBusFinalValuesUpdateMessage(const UE::Trace::IAnalyzer::FOnEventContext& InContext) : FModulationMatrixUpdateMessageBase(InContext) {}
+
+		// from UE::Audio::Insights::IAudioCachedMessage
+		virtual const FName GetMessageName() const override { return ModulationMatrixMessageNames::BusFinalValuesUpdate; }
+		virtual uint32 GetSizeOf() const override { return sizeof(FBusFinalValuesUpdateMessage) + BusIdToValueMap.GetAllocatedSize(); }
+	};
 
 	// Deactivate message
-	using FModulationMatrixDeactivateMessage = FModulationMatrixMessageBase;
+	struct FModulationMatrixDeactivateMessage final : public FModulationMatrixMessageBase
+	{
+		FModulationMatrixDeactivateMessage() = default;
+		FModulationMatrixDeactivateMessage(const UE::Trace::IAnalyzer::FOnEventContext& InContext) : FModulationMatrixMessageBase(InContext) {}
+
+		// from UE::Audio::Insights::IAudioCachedMessage
+		virtual const FName GetMessageName() const override { return ModulationMatrixMessageNames::ModulatingSourceDeactivate; }
+		virtual uint32 GetSizeOf() const override { return sizeof(FModulationMatrixDeactivateMessage); }
+	};
 
 
 	// Modulation matrix messages
diff --git a/Engine/Plugins/Runtime/AudioModulationInsights/Source/AudioModulationInsights/Private/Providers/ModulationMatrixTraceProvider.cpp b/Engine/Plugins/Runtime/AudioModulationInsights/Source/AudioModulationInsights/Private/Providers/ModulationMatrixTraceProvider.cpp
--- a/Engine/Plugins/Runtime/AudioModulationInsights/Source/AudioModulationInsights/Private/Providers/ModulationMatrixTraceProvider.cpp
+++ b/Engine/Plugins/Runtime/AudioModulationInsights/Source/AudioModulationInsights/Private/Providers/ModulationMatrixTraceProvider.cpp
@@ -291,43 +291,43 @@ namespace AudioModulationInsights
 					case RouteId_BusMixRegisterBus:
 					case RouteId_GeneratorRegisterBus:
 					{
-						Messages.RegisterBusMessages.Enqueue(FModulationMatrixRegisterBusMessage{ Context });
+						CacheMessage<FModulationMatrixRegisterBusMessage>(Context, Messages.RegisterBusMessages);
 						break;
 					}
 
 					case RouteId_BusMixActivate:
 					{
-						Messages.BusMixActivateMessages.Enqueue(FBusMixActivateMessage{ Context });
+						CacheMessage<FBusMixActivateMessage>(Context, Messages.BusMixActivateMessages);
 						break;
 					}
 
 					case RouteId_GeneratorActivate:
 					{
-						Messages.GeneratorActivateMessages.Enqueue(FGeneratorActivateMessage{ Context });
+						CacheMessage<FGeneratorActivateMessage>(Context, Messages.GeneratorActivateMessages);
 						break;
 					}
 
 					case RouteId_BusMixUpdate:
 					{
-						Messages.BusMixUpdateMessages.Enqueue(FBusMixUpdateMessage{ Context });
+						CacheMessage<FBusMixUpdateMessage>(Context, Messages.BusMixUpdateMessages);
 						break;
 					}
 
 					case RouteId_GeneratorUpdate:
 					{
-						Messages.GeneratorUpdateMessages.Enqueue(FGeneratorUpdateMessage{ Context });
+						CacheMessage<FGeneratorUpdateMessage>(Context, Messages.GeneratorUpdateMessages);
 						break;
 					}
 
 					case RouteId_BusFinalValuesUpdate:
 					{
-						Messages.BusFinalValuesUpdateMessages.Enqueue(FBusFinalValuesUpdateMessage{ Context });
+						CacheMessage<FBusFinalValuesUpdateMessage>(Context, Messages.BusFinalValuesUpdateMessages);
 						break;
 					}
 
 					case RouteId_Deactivate:
 					{
-						Messages.DeactivateMessages.Enqueue(FModulationMatrixDeactivateMessage{ Context });
+						CacheMessage<FModulationMatrixDeactivateMessage>(Context, Messages.DeactivateMessages);
 						break;
 					}
 
