[Audio Insights] Clean up Control Bus and Modulation Matrix Dashboards Pt.1

* Minor rename of Control Bus column names for clarity
* Part of: https://jira.it.epicgames.com/browse/UE-356383

#Synced-CL 49103850

diff --git a/Engine/Plugins/Runtime/AudioModulationInsights/Source/AudioModulationInsights/Private/Views/ControlBusDashboardViewFactory.cpp b/Engine/Plugins/Runtime/AudioModulationInsights/Source/AudioModulationInsights/Private/Views/ControlBusDashboardViewFactory.cpp
--- a/Engine/Plugins/Runtime/AudioModulationInsights/Source/AudioModulationInsights/Private/Views/ControlBusDashboardViewFactory.cpp
+++ b/Engine/Plugins/Runtime/AudioModulationInsights/Source/AudioModulationInsights/Private/Views/ControlBusDashboardViewFactory.cpp
@@ -115,7 +115,7 @@ namespace AudioModulationInsights
 				{
 					"Name",
 					{
-						LOCTEXT("ControlBus_NameColumnDisplayName", "Name"),
+						LOCTEXT("ControlBus_NameColumnDisplayName", "Control Bus Name"),
 						[](const IDashboardDataViewEntry& InData) { return ControlBusPrivate::CastEntry(InData).GetDisplayNameAsFText(); },
 						nullptr	/* GetIconName */,
 						false /* bDefaultHidden */,
@@ -125,7 +125,7 @@ namespace AudioModulationInsights
 				{
 					"ParameterName",
 					{
-						LOCTEXT("ControlBus_ParamNameColumnDisplayName", "Parameter"),
+						LOCTEXT("ControlBus_ParamNameColumnDisplayName", "Modulation Parameter"),
 						[](const IDashboardDataViewEntry& InData) { return FText::FromString(ControlBusPrivate::CastEntry(InData).ParamName); },
 						nullptr	/* GetIconName */,
 						false /* bDefaultHidden */,
