[AudioInsights] - AudioModulation: Consolidate mirrored engine data types to single header

Passes BugHawk

#Synced-CL 50159697

diff --git a/Engine/Plugins/Runtime/AudioModulationInsights/Source/AudioModulationInsights/Private/AudioModulationDataTypes.h b/Engine/Plugins/Runtime/AudioModulationInsights/Source/AudioModulationInsights/Private/AudioModulationDataTypes.h
new file mode 100644
--- /dev/null
+++ b/Engine/Plugins/Runtime/AudioModulationInsights/Source/AudioModulationInsights/Private/AudioModulationDataTypes.h
@@ -0,0 +1,21 @@
+// Copyright Epic Games, Inc. All Rights Reserved.
+#pragma once
+
+#include "HAL/Platform.h"
+
+namespace AudioModulationInsights
+{
+    /*
+    * Note: The following mirrors data from Engine/Plugins/Runtime/AudioModulation
+    *
+    * Duplication is required as AudioInsights cannot depend on any engine code
+    * As it must be able to run in standalone (where engine code is not available)
+    * 
+    * Caution: Since this means we no longer have a single source of truth
+    * We must maintain this data if the underlying engine types change
+    */
+
+    using FBusId = uint32;          // mirrors FBusId: Engine/Plugins/Runtime/AudioModulation/Source/AudioModulation/Private/SoundControlBusProxy.h
+    using FModulatorId = uint32;    // mirrors FModulatorId: Engine/Source/Runtime/AudioExtensions/Public/IAudioModulation.h
+
+} // namespace AudioModulationInsights
diff --git a/Engine/Plugins/Runtime/AudioModulationInsights/Source/AudioModulationInsights/Private/Messages/ControlBusTraceMessages.h b/Engine/Plugins/Runtime/AudioModulationInsights/Source/AudioModulationInsights/Private/Messages/ControlBusTraceMessages.h
--- a/Engine/Plugins/Runtime/AudioModulationInsights/Source/AudioModulationInsights/Private/Messages/ControlBusTraceMessages.h
+++ b/Engine/Plugins/Runtime/AudioModulationInsights/Source/AudioModulationInsights/Private/Messages/ControlBusTraceMessages.h
@@ -2,6 +2,7 @@
 #pragma once
 
 #include "AudioDefines.h"
+#include "AudioModulationDataTypes.h"
 #include "Cache/IAudioCachedMessage.h"
 #include "Messages/AnalyzerMessageQueue.h"
 #include "Trace/Analyzer.h"
@@ -9,9 +10,6 @@
 
 namespace AudioModulationInsights
 {
-
-	using FBusId = uint32;
-
 	namespace ControlBusMessageNames
 	{
 		inline const FName Activate(TEXT("ControlBusActivate"));
diff --git a/Engine/Plugins/Runtime/AudioModulationInsights/Source/AudioModulationInsights/Private/Messages/ModulationMatrixTraceMessages.h b/Engine/Plugins/Runtime/AudioModulationInsights/Source/AudioModulationInsights/Private/Messages/ModulationMatrixTraceMessages.h
--- a/Engine/Plugins/Runtime/AudioModulationInsights/Source/AudioModulationInsights/Private/Messages/ModulationMatrixTraceMessages.h
+++ b/Engine/Plugins/Runtime/AudioModulationInsights/Source/AudioModulationInsights/Private/Messages/ModulationMatrixTraceMessages.h
@@ -2,6 +2,7 @@
 #pragma once
 
 #include "AudioDefines.h"
+#include "AudioModulationDataTypes.h"
 #include "Cache/IAudioCachedMessage.h"
 #include "Messages/AnalyzerMessageQueue.h"
 #include "Trace/Analyzer.h"
@@ -9,9 +10,6 @@
 
 namespace AudioModulationInsights
 {
-	using FBusId    = uint32;
-	using FSourceId = uint32;
-
 	enum class EModulationMatrixEntryType : uint8
 	{
 		BusMix,
@@ -38,7 +36,7 @@ namespace AudioModulationInsights
 		FModulationMatrixDashboardEntry() = default;
 		virtual ~FModulationMatrixDashboardEntry() = default;
 
-		FSourceId SourceId = INDEX_NONE;
+		FModulatorId SourceId = INDEX_NONE;
 		EModulationMatrixEntryType EntryType = EModulationMatrixEntryType::None;
 		FString DisplayName;
 
@@ -57,14 +55,14 @@ namespace AudioModulationInsights
 
 			DeviceId  = static_cast<::Audio::FDeviceId>(EventData.GetValue<uint32>("DeviceId"));
 			Timestamp = InContext.EventTime.AsSeconds(EventData.GetValue<uint64>("Timestamp"));
-			SourceId  = static_cast<FSourceId>(EventData.GetValue<uint32>("SourceId"));
+			SourceId  = static_cast<FModulatorId>(EventData.GetValue<uint32>("SourceId"));
 		}
 
 		// from UE::Audio::Insights::IAudioCachedMessage
 		virtual uint64 GetID() const override { return SourceId; }
 
 		Audio::FDeviceId DeviceId = INDEX_NONE;
-		FSourceId SourceId = INDEX_NONE;
+		FModulatorId SourceId = INDEX_NONE;
 	};
 
 
@@ -78,7 +76,7 @@ namespace AudioModulationInsights
 		{
 			const UE::Trace::IAnalyzer::FEventData& EventData = InContext.EventData;
 
-			ModulatingSourceId = static_cast<FSourceId>(EventData.GetValue<uint32>("ModulatingSourceId"));
+			ModulatingSourceId = static_cast<FModulatorId>(EventData.GetValue<uint32>("ModulatingSourceId"));
 			EventData.GetString("BusName", BusName);
 		}
 
@@ -86,7 +84,7 @@ namespace AudioModulationInsights
 		virtual const FName GetMessageName() const override { return ModulationMatrixMessageNames::BusMixRegisterBus; }
 		virtual uint32 GetSizeOf() const override { return sizeof(FModulationMatrixRegisterBusMessage) + BusName.GetAllocatedSize(); }
 
-		FSourceId ModulatingSourceId = INDEX_NONE;
+		FModulatorId ModulatingSourceId = INDEX_NONE;
 		FString BusName;
 	};
 
diff --git a/Engine/Plugins/Runtime/AudioModulationInsights/Source/AudioModulationInsights/Private/Providers/ModulationMatrixTraceProvider.cpp b/Engine/Plugins/Runtime/AudioModulationInsights/Source/AudioModulationInsights/Private/Providers/ModulationMatrixTraceProvider.cpp
--- a/Engine/Plugins/Runtime/AudioModulationInsights/Source/AudioModulationInsights/Private/Providers/ModulationMatrixTraceProvider.cpp
+++ b/Engine/Plugins/Runtime/AudioModulationInsights/Source/AudioModulationInsights/Private/Providers/ModulationMatrixTraceProvider.cpp
@@ -1,6 +1,7 @@
 // Copyright Epic Games, Inc. All Rights Reserved.
 #include "ModulationMatrixTraceProvider.h"
 
+#include "AudioModulationDataTypes.h"
 #include "AudioInsightsUtils.h"
 #include "Trace/Analyzer.h"
 #include "TraceServices/Model/AnalysisSession.h"
@@ -118,9 +119,9 @@ namespace AudioModulationInsights
 
 		if (bHasRegisterBusMsgs && (bHasBusMixActivateMsgs || bHasGeneratorActivateMsgs))
 		{
-			TMap<FSourceId, TMap<FBusId, FString>> ModulatingSourceIdToBusesInfoMap;
+			TMap<FModulatorId, TMap<FBusId, FString>> ModulatingSourceIdToBusesInfoMap;
 
-			auto UpdateActiveModulatorSourceIds = [&](const Audio::FDeviceId DeviceId, const FSourceId SourceId)
+			auto UpdateActiveModulatorSourceIds = [&](const Audio::FDeviceId DeviceId, const FModulatorId SourceId)
 			{
 				TSet<uint32>& ActiveModulatorSourceIds = DeviceIdToActiveModulatorSourceIdsMap.FindOrAdd(DeviceId);
 
diff --git a/Engine/Plugins/Runtime/AudioModulationInsights/Source/AudioModulationInsights/Private/Views/ModulationMatrixDashboardViewFactory.cpp b/Engine/Plugins/Runtime/AudioModulationInsights/Source/AudioModulationInsights/Private/Views/ModulationMatrixDashboardViewFactory.cpp
--- a/Engine/Plugins/Runtime/AudioModulationInsights/Source/AudioModulationInsights/Private/Views/ModulationMatrixDashboardViewFactory.cpp
+++ b/Engine/Plugins/Runtime/AudioModulationInsights/Source/AudioModulationInsights/Private/Views/ModulationMatrixDashboardViewFactory.cpp
@@ -1,6 +1,7 @@
 // Copyright Epic Games, Inc. All Rights Reserved.
 #include "ModulationMatrixDashboardViewFactory.h"
 
+#include "AudioModulationDataTypes.h"
 #include "AudioInsightsStyle.h"
 #include "AudioInsightsTraceProviderBase.h"
 #include "AudioModulationInsightsStyle.h"
@@ -104,7 +105,7 @@ namespace AudioModulationInsights
 					LOCTEXT("ModulationMatrix_SourceIdColumnDisplayName", "Source Id"),
 					[](const IDashboardDataViewEntry& InData)
 					{
-						const FSourceId SourceId = CastEntry(InData).SourceId;
+						const FModulatorId SourceId = CastEntry(InData).SourceId;
 						return SourceId != INDEX_NONE ? FText::AsNumber(SourceId) : FText::GetEmpty();
 					},
 					nullptr	/* GetIconName */,
