[AudioInsights] - Remove obsolete ControlBus & ModulationMatrix code and files

#Synced-CL 50644355

diff --git a/Engine/Plugins/Runtime/AudioModulationInsights/Source/AudioModulationInsights/Private/AudioModulationInsightsModule.cpp b/Engine/Plugins/Runtime/AudioModulationInsights/Source/AudioModulationInsights/Private/AudioModulationInsightsModule.cpp
--- a/Engine/Plugins/Runtime/AudioModulationInsights/Source/AudioModulationInsights/Private/AudioModulationInsightsModule.cpp
+++ b/Engine/Plugins/Runtime/AudioModulationInsights/Source/AudioModulationInsights/Private/AudioModulationInsightsModule.cpp
@@ -2,8 +2,6 @@
 #include "AudioModulationInsightsModule.h"
 
 #include "IAudioInsightsModule.h"
-#include "Views/ControlBusDashboardViewFactory.h"
-#include "Views/ModulationMatrixDashboardViewFactory.h"
 
 #if WITH_EDITOR
 #include "IAudioInsightsEditorModule.h"
@@ -19,9 +17,6 @@ void FAudioModulationInsightsModule::StartupModule()
 #else
 		IAudioInsightsModule& AudioInsightsModule = IAudioInsightsModule::GetChecked();
 #endif // WITH_EDITOR
-		
-		AudioInsightsModule.RegisterDashboardViewFactory(MakeShared<AudioModulationInsights::FControlBusDashboardViewFactory>());
-		AudioInsightsModule.RegisterDashboardViewFactory(MakeShared<AudioModulationInsights::FModulationMatrixDashboardViewFactory>());
 
 		AudioModulationDashboardViewFactory = MakeShared<AudioModulationInsights::FAudioModulationDashboardViewFactory>();
 		AudioInsightsModule.RegisterDashboardViewFactory(AudioModulationDashboardViewFactory.ToSharedRef());
@@ -47,9 +42,6 @@ void FAudioModulationInsightsModule::ShutdownModule()
 
 		IAudioInsightsModule& AudioInsightsModule = IAudioInsightsModule::GetChecked();
 	#endif // WITH_EDITOR
-			
-		AudioInsightsModule.UnregisterDashboardViewFactory("ControlBuses");
-		AudioInsightsModule.UnregisterDashboardViewFactory("ModulationMatrix");
 
 		if (AudioModulationDashboardViewFactory.IsValid())
 		{
diff --git a/Engine/Plugins/Runtime/AudioModulationInsights/Source/AudioModulationInsights/Private/Messages/ControlBusTraceMessages.h b/Engine/Plugins/Runtime/AudioModulationInsights/Source/AudioModulationInsights/Private/Messages/ControlBusTraceMessages.h
deleted file mode 100644
diff --git a/Engine/Plugins/Runtime/AudioModulationInsights/Source/AudioModulationInsights/Private/Messages/ModulationMatrixTraceMessages.h b/Engine/Plugins/Runtime/AudioModulationInsights/Source/AudioModulationInsights/Private/Messages/ModulationMatrixTraceMessages.h
deleted file mode 100644
diff --git a/Engine/Plugins/Runtime/AudioModulationInsights/Source/AudioModulationInsights/Private/Providers/ControlBusTraceProvider.cpp b/Engine/Plugins/Runtime/AudioModulationInsights/Source/AudioModulationInsights/Private/Providers/ControlBusTraceProvider.cpp
deleted file mode 100644
diff --git a/Engine/Plugins/Runtime/AudioModulationInsights/Source/AudioModulationInsights/Private/Providers/ControlBusTraceProvider.h b/Engine/Plugins/Runtime/AudioModulationInsights/Source/AudioModulationInsights/Private/Providers/ControlBusTraceProvider.h
deleted file mode 100644
diff --git a/Engine/Plugins/Runtime/AudioModulationInsights/Source/AudioModulationInsights/Private/Providers/ModulationMatrixTraceProvider.cpp b/Engine/Plugins/Runtime/AudioModulationInsights/Source/AudioModulationInsights/Private/Providers/ModulationMatrixTraceProvider.cpp
deleted file mode 100644
diff --git a/Engine/Plugins/Runtime/AudioModulationInsights/Source/AudioModulationInsights/Private/Providers/ModulationMatrixTraceProvider.h b/Engine/Plugins/Runtime/AudioModulationInsights/Source/AudioModulationInsights/Private/Providers/ModulationMatrixTraceProvider.h
deleted file mode 100644
diff --git a/Engine/Plugins/Runtime/AudioModulationInsights/Source/AudioModulationInsights/Private/Views/ControlBusDashboardViewFactory.cpp b/Engine/Plugins/Runtime/AudioModulationInsights/Source/AudioModulationInsights/Private/Views/ControlBusDashboardViewFactory.cpp
deleted file mode 100644
diff --git a/Engine/Plugins/Runtime/AudioModulationInsights/Source/AudioModulationInsights/Private/Views/ControlBusDashboardViewFactory.h b/Engine/Plugins/Runtime/AudioModulationInsights/Source/AudioModulationInsights/Private/Views/ControlBusDashboardViewFactory.h
deleted file mode 100644
diff --git a/Engine/Plugins/Runtime/AudioModulationInsights/Source/AudioModulationInsights/Private/Views/ModulationMatrixDashboardViewFactory.cpp b/Engine/Plugins/Runtime/AudioModulationInsights/Source/AudioModulationInsights/Private/Views/ModulationMatrixDashboardViewFactory.cpp
deleted file mode 100644
diff --git a/Engine/Plugins/Runtime/AudioModulationInsights/Source/AudioModulationInsights/Private/Views/ModulationMatrixDashboardViewFactory.h b/Engine/Plugins/Runtime/AudioModulationInsights/Source/AudioModulationInsights/Private/Views/ModulationMatrixDashboardViewFactory.h
deleted file mode 100644
