Fix empty source LOCTEXT macro warning in AudioModulationDashboardViewFactory

#Synced-CL 50838485

diff --git a/Engine/Plugins/Runtime/AudioModulationInsights/Source/AudioModulationInsights/Private/Views/AudioModulationDashboardViewFactory.cpp b/Engine/Plugins/Runtime/AudioModulationInsights/Source/AudioModulationInsights/Private/Views/AudioModulationDashboardViewFactory.cpp
--- a/Engine/Plugins/Runtime/AudioModulationInsights/Source/AudioModulationInsights/Private/Views/AudioModulationDashboardViewFactory.cpp
+++ b/Engine/Plugins/Runtime/AudioModulationInsights/Source/AudioModulationInsights/Private/Views/AudioModulationDashboardViewFactory.cpp
@@ -315,7 +315,7 @@ namespace AudioModulationInsights
 							const FAudioModulationDashboardEntry& CastedInData = FAudioModulationDashboardEntry::CastEntry(InData);
 							if (CastedInData.ModulatorType == EModulatorTraceType::ControlBusMix)
 							{
-								return LOCTEXT("AudioModulation_ModulatorValueBusMix", "");
+								return FText() /* LOCTEXT("AudioModulation_ModulatorValueBusMix", "") */;
 							}
 							
 							return FText::AsNumber(CastedInData.Value, FSlateStyle::Get().GetLinearVolumeFloatFormat()); 
