[Audio Insights] - Modulation: Fix Bughawk error in AudioModulationInsightsStyle.cpp

#Synced-CL 51664040

diff --git a/Engine/Plugins/Developer/AudioModulationInsights/Source/AudioModulationInsights/Private/AudioModulationInsightsStyle.cpp b/Engine/Plugins/Developer/AudioModulationInsights/Source/AudioModulationInsights/Private/AudioModulationInsightsStyle.cpp
--- a/Engine/Plugins/Developer/AudioModulationInsights/Source/AudioModulationInsights/Private/AudioModulationInsightsStyle.cpp
+++ b/Engine/Plugins/Developer/AudioModulationInsights/Source/AudioModulationInsights/Private/AudioModulationInsightsStyle.cpp
@@ -19,8 +19,11 @@ FAudioModulationInsightsStyle::FAudioModulationInsightsStyle()
 
 	// Register style content paths
 	const TSharedPtr<const IPlugin> Plugin = IPluginManager::Get().FindPlugin(UE_PLUGIN_NAME);
+	if (Plugin.IsValid())
+	{
+		SetContentRoot(Plugin->GetBaseDir() / "Resources");
+	}
 
-	SetContentRoot(Plugin->GetBaseDir() / "Resources");
 	SetCoreContentRoot(FPaths::EngineContentDir() / "Slate");
 
 	// Icons
